Missing file from CL# 2531273

diff --git a/Engine/Source/Developer/MaterialUtilities/Private/MaterialUtilitiesPrivatePCH.h b/Engine/Source/Developer/MaterialUtilities/Private/MaterialUtilitiesPrivatePCH.h
new file mode 100644
--- /dev/null
+++ b/Engine/Source/Developer/MaterialUtilities/Private/MaterialUtilitiesPrivatePCH.h
@@ -0,0 +1,2 @@
+// Copyright 1998-2015 Epic Games, Inc. All Rights Reserved.
+#include "MaterialUtilities.h"
