fixing mac editor build

diff --git a/Engine/Source/Developer/MaterialUtilities/Public/MaterialUtilities.h b/Engine/Source/Developer/MaterialUtilities/Public/MaterialUtilities.h
--- a/Engine/Source/Developer/MaterialUtilities/Public/MaterialUtilities.h
+++ b/Engine/Source/Developer/MaterialUtilities/Public/MaterialUtilities.h
@@ -3,6 +3,10 @@
 #pragma once
 
 #include "ModuleInterface.h"
+#include "Runtime/CoreUObject/Public/UObject/ObjectBase.h"
+#include "Runtime/Engine/Classes/Engine/EngineTypes.h"
+#include "Runtime/Engine/Classes/Engine/Texture.h"
+#include "Runtime/Engine/Public/SceneTypes.h"
 
 struct FFlattenMaterial
 {
@@ -37,6 +41,7 @@ class UTextureRenderTarget2D;
 class UWorld;
 class ALandscapeProxy;
 class ULandscapeComponent;
+class FPrimitiveComponentId;
 
 /**
  * Material utilities
