fixing mac editor build

diff --git a/Engine/Source/Developer/MaterialUtilities/Private/MaterialUtilities.cpp b/Engine/Source/Developer/MaterialUtilities/Private/MaterialUtilities.cpp
--- a/Engine/Source/Developer/MaterialUtilities/Private/MaterialUtilities.cpp
+++ b/Engine/Source/Developer/MaterialUtilities/Private/MaterialUtilities.cpp
@@ -6,15 +6,19 @@
 #include "Runtime/Engine/Classes/Materials/MaterialExpressionConstant.h"
 #include "Runtime/Engine/Classes/Engine/TextureRenderTarget2D.h"
 #include "Runtime/Engine/Classes/Engine/Texture2D.h"
+#include "Runtime/Engine/Classes/Engine/TextureCube.h"
 #include "Runtime/Engine/Public/TileRendering.h"
 #include "Runtime/Engine/Public/EngineModule.h"
 #include "Runtime/Engine/Public/ImageUtils.h"
+#include "Runtime/Engine/Public/CanvasTypes.h"
 #include "Runtime/Engine/Public/MaterialCompiler.h"
+#include "Runtime/Engine/Classes/Engine/TextureLODSettings.h"
+#include "Runtime/Engine/Classes/DeviceProfiles/DeviceProfileManager.h"
 #include "RendererInterface.h"
 #include "LandscapeProxy.h"
 #include "LandscapeComponent.h"
 
-DEFINE_LOG_CATEGORY_STATIC(LogMeshUtilities, Log, All);
+DEFINE_LOG_CATEGORY_STATIC(LogMaterialUtilities, Log, All);
 
 /*------------------------------------------------------------------------------
 	Helper classes for render material to texture
@@ -905,7 +909,7 @@ UMaterial* FMaterialUtilities::CreateMaterial(const FFlattenMaterial& InFlattenM
 	// In case outer is null BaseName has to be long package name
 	if (InOuter == nullptr && FPackageName::IsShortPackageName(BaseName))
 	{
-		UE_LOG(LogMeshUtilities, Warning, TEXT("Invalid long package name: '%s'."), *BaseName);
+		UE_LOG(LogMaterialUtilities, Warning, TEXT("Invalid long package name: '%s'."), *BaseName);
 		return nullptr;
 	}
 
