Fixed crash in exporting base color texture for landscape

diff --git a/Engine/Source/Developer/MaterialUtilities/Private/MaterialUtilities.cpp b/Engine/Source/Developer/MaterialUtilities/Private/MaterialUtilities.cpp
--- a/Engine/Source/Developer/MaterialUtilities/Private/MaterialUtilities.cpp
+++ b/Engine/Source/Developer/MaterialUtilities/Private/MaterialUtilities.cpp
@@ -18,6 +18,10 @@
 #include "LandscapeProxy.h"
 #include "LandscapeComponent.h"
 
+
+IMPLEMENT_MODULE(FDefaultModuleImpl, MaterialUtilities);
+
+
 DEFINE_LOG_CATEGORY_STATIC(LogMaterialUtilities, Log, All);
 
 /*------------------------------------------------------------------------------
