Fixed: HLOD assets have lighter materials than original meshes
UE-14625

diff --git a/Engine/Source/Developer/MaterialUtilities/Private/MaterialUtilities.cpp b/Engine/Source/Developer/MaterialUtilities/Private/MaterialUtilities.cpp
--- a/Engine/Source/Developer/MaterialUtilities/Private/MaterialUtilities.cpp
+++ b/Engine/Source/Developer/MaterialUtilities/Private/MaterialUtilities.cpp
@@ -27,6 +27,72 @@ DEFINE_LOG_CATEGORY_STATIC(LogMaterialUtilities, Log, All);
 /*------------------------------------------------------------------------------
 	Helper classes for render material to texture
 ------------------------------------------------------------------------------*/
+struct FExportMaterialCompiler : public FProxyMaterialCompiler
+{
+	FExportMaterialCompiler(FMaterialCompiler* InCompiler) :
+		FProxyMaterialCompiler(InCompiler)
+	{}
+
+	// gets value stored by SetMaterialProperty()
+	virtual EShaderFrequency GetCurrentShaderFrequency() const override
+	{
+		// not used by Lightmass
+		return SF_Pixel;
+	}
+
+	virtual int32 ParticleRelativeTime() override
+	{
+		return Compiler->Constant(0.0f);
+	}
+
+	virtual int32 ParticleMotionBlurFade() override
+	{
+		return Compiler->Constant(1.0f);
+	}
+
+	virtual int32 ParticleDirection() override
+	{
+		return Compiler->Constant3(0.0f, 0.0f, 0.0f);
+	}
+
+	virtual int32 ParticleSpeed() override
+	{
+		return Compiler->Constant(0.0f);
+	}
+	
+	virtual int32 ParticleSize() override
+	{
+		return Compiler->Constant2(0.0f,0.0f);
+	}
+
+	virtual int32 ObjectRadius() override
+	{
+		return Compiler->Constant(500);
+	}
+
+	virtual int32 ObjectBounds() override
+	{
+		return Compiler->Constant3(0,0,0);
+	}
+
+	virtual int32 CameraVector() override
+	{
+		return Compiler->Constant3(0.0f,0.0f,1.0f);
+	}
+
+	virtual int32 ReflectionAboutCustomWorldNormal(int32 CustomWorldNormal, int32 bNormalizeCustomWorldNormal) override
+	{
+		return Compiler->Constant3(0.0f,0.0f,-1.0f);
+	}
+
+	virtual int32 VertexColor() override
+	{
+		return Compiler->Constant4(1.0f,1.0f,1.0f,1.0f);
+	}
+};
+
+
+
 class FExportMaterialProxy : public FMaterial, public FMaterialRenderProxy
 {
 public:
@@ -145,19 +211,21 @@ public:
 		if (Property == MP_EmissiveColor)
 		{
 			UMaterial* ProxyMaterial = MaterialInterface->GetMaterial();
+			check(ProxyMaterial);
 			EBlendMode BlendMode = MaterialInterface->GetBlendMode();
 			EMaterialShadingModel ShadingModel = MaterialInterface->GetShadingModel();
-			check(ProxyMaterial);
+			FExportMaterialCompiler ProxyCompiler(Compiler);
+									
 			switch (PropertyToCompile)
 			{
 			case MP_EmissiveColor:
 				// Emissive is ALWAYS returned...
-				return Compiler->ForceCast(MaterialInterface->CompileProperty(Compiler,MP_EmissiveColor),MCT_Float3,true,true);
+				return Compiler->ForceCast(MaterialInterface->CompileProperty(&ProxyCompiler,MP_EmissiveColor),MCT_Float3,true,true);
 			case MP_BaseColor:
 				// Only return for Opaque and Masked...
 				if (BlendMode == BLEND_Opaque || BlendMode == BLEND_Masked)
 				{
-					return Compiler->ForceCast(MaterialInterface->CompileProperty(Compiler, MP_BaseColor),MCT_Float3,true,true);
+					return Compiler->ForceCast(MaterialInterface->CompileProperty(&ProxyCompiler, MP_BaseColor),MCT_Float3,true,true);
 				}
 				break;
 			case MP_Specular: 
@@ -166,7 +234,7 @@ public:
 				// Only return for Opaque and Masked...
 				if (BlendMode == BLEND_Opaque || BlendMode == BLEND_Masked)
 				{
-					return Compiler->ForceCast(MaterialInterface->CompileProperty(Compiler, PropertyToCompile),MCT_Float,true,true);
+					return Compiler->ForceCast(MaterialInterface->CompileProperty(&ProxyCompiler, PropertyToCompile),MCT_Float,true,true);
 				}
 				break;
 			case MP_Normal:
@@ -175,7 +243,7 @@ public:
 				{
 					return Compiler->ForceCast( 
 						Compiler->Add( 
-							Compiler->Mul(MaterialInterface->CompileProperty(Compiler, MP_Normal), Compiler->Constant(0.5f)), // [-1,1] * 0.5
+							Compiler->Mul(MaterialInterface->CompileProperty(&ProxyCompiler, MP_Normal), Compiler->Constant(0.5f)), // [-1,1] * 0.5
 							Compiler->Constant(0.5f)), // [-0.5,0.5] + 0.5
 						MCT_Float3, true, true );
 				}
@@ -706,8 +774,7 @@ bool FMaterialUtilities::ExportMaterial(UWorld* InWorld, UMaterialInterface* InM
 
 	// Render normal property
 	if (OutFlattenMaterial.NormalSize.X > 0 && 
-		OutFlattenMaterial.NormalSize.Y > 0 &&
-		InMaterial->GetMaterial()->HasNormalConnected())
+		OutFlattenMaterial.NormalSize.Y > 0)
 	{
 		// Create temporary render target
 		auto RenderTargetNormal = NewObject<UTextureRenderTarget2D>();
