Added more material simplification options (constants, texture sizes) for LOD maps generation
UE-16826

diff --git a/Engine/Source/Developer/MaterialUtilities/Private/MaterialUtilities.cpp b/Engine/Source/Developer/MaterialUtilities/Private/MaterialUtilities.cpp
--- a/Engine/Source/Developer/MaterialUtilities/Private/MaterialUtilities.cpp
+++ b/Engine/Source/Developer/MaterialUtilities/Private/MaterialUtilities.cpp
@@ -348,98 +348,6 @@ public:
 		return Ar << V.MaterialInterface;
 	}
 
-	bool IsMaterialInputConnected(UMaterial* InMaterial, EMaterialProperty MaterialInput)
-	{
-		bool bConnected = false;
-
-		switch (MaterialInput)
-		{
-		case MP_EmissiveColor:
-			bConnected = InMaterial->EmissiveColor.Expression != NULL;
-			break;
-		case MP_BaseColor:
-			bConnected = InMaterial->BaseColor.Expression != NULL;
-			break;
-		case MP_Specular:
-			bConnected = InMaterial->Specular.Expression != NULL;
-			break;
-		case MP_Normal:
-			bConnected = InMaterial->Normal.Expression != NULL;
-			break;
-		case MP_Metallic:
-			bConnected = InMaterial->Metallic.Expression != NULL;
-			break;
-		case MP_Roughness:
-			bConnected = InMaterial->Roughness.Expression != NULL;
-			break;
-		default:
-			break;
-		}
-
-		// Note: only checking to see whether the entire material attributes connection exists.  
-		// This means materials using the material attributes input will export more attributes than is necessary.
-		bConnected = InMaterial->bUseMaterialAttributes ? InMaterial->MaterialAttributes.Expression != NULL : bConnected;
-		return bConnected;
-	}
-
-	/**
-	 *	Checks if the configuration of the material proxy will generate a uniform
-	 *	value across the sampling... (Ie, nothing is hooked to the property)
-	 *
-	 *	@param	OutUniformValue		The value that will be returned.
-	 *
-	 *	@return	bool				true if a single value would be generated.
-	 *								false if not.
-	 */
-	bool WillGenerateUniformData(FColor& OutUniformValue)
-	{
-		// Pre-fill the value...
-		OutUniformValue.R = 0;
-		OutUniformValue.G = 0;
-		OutUniformValue.B = 0;
-		OutUniformValue.A = 0;
-
-		EBlendMode BlendMode = MaterialInterface->GetBlendMode();
-		EMaterialShadingModel ShadingModel = MaterialInterface->GetShadingModel();
-		
-		check(Material);
-		bool bExpressionIsNULL = false;
-		switch (PropertyToCompile)
-		{
-		case MP_EmissiveColor:
-			// Emissive is ALWAYS returned...
-			bExpressionIsNULL = !IsMaterialInputConnected(Material, PropertyToCompile);
-			break;
-		case MP_BaseColor:
-			// Only return for Opaque and Masked...
-			if (BlendMode == BLEND_Opaque || BlendMode == BLEND_Masked)
-			{
-				bExpressionIsNULL = !IsMaterialInputConnected(Material, PropertyToCompile);
-			}
-			break;
-		case MP_Specular: 
-		case MP_Metallic:
-		case MP_Roughness:
-			// Only return for Opaque and Masked...
-			if (BlendMode == BLEND_Opaque || BlendMode == BLEND_Masked)
-			{
-				bExpressionIsNULL = !IsMaterialInputConnected(Material, PropertyToCompile);
-				OutUniformValue.A = 255;
-			}
-			break;
-		case MP_Normal:
-			// Only return for Opaque and Masked...
-			if (BlendMode == BLEND_Opaque || BlendMode == BLEND_Masked)
-			{
-				bExpressionIsNULL = !IsMaterialInputConnected(Material, PropertyToCompile);
-				OutUniformValue.B = 255;	// Default normal is (0,0,1)
-			}
-			break;
-		}
-
-		return bExpressionIsNULL;
-	}
-
 	/**
 	* Iterate through all textures used by the material and return the maximum texture resolution used
 	* (ideally this could be made dependent of the material property)
@@ -677,15 +585,6 @@ bool FMaterialUtilities::ExportMaterialProperty(UWorld* InWorld, UMaterialInterf
 		return false;
 	}
 
-	FColor UniformValue;
-	if (MaterialProxy->WillGenerateUniformData(UniformValue))
-	{
-		// Single value... fill it in.
-		OutBMP.Empty();
-		OutBMP.Add(UniformValue);
-		return true;
-	}
-
 	bool bNormalmap = (InMaterialProperty == MP_Normal);
 
 	RenderMaterialTile(InWorld, MaterialProxy, InRenderTarget, !bNormalmap);
@@ -705,15 +604,6 @@ bool FMaterialUtilities::ExportMaterialProperty(UWorld* InWorld, UMaterialInterf
 		return false;
 	}
 
-	FColor UniformValue;
-	if (MaterialProxy->WillGenerateUniformData(UniformValue))
-	{
-		// Single value... fill it in.
-		OutBMP.Empty();
-		OutBMP.Add(UniformValue);
-		return true;
-	}
-
 	bool bNormalmap = (InMaterialProperty == MP_Normal);
 	EPixelFormat Format = PF_FloatRGB;
 	OutSize = MaterialProxy->FindMaxTextureSize(InMaterial);
@@ -1054,10 +944,10 @@ UMaterial* FMaterialUtilities::CreateMaterial(const FFlattenMaterial& InFlattenM
 
 		MaterialNodeY+= MaterialNodeStepY;
 	}
-	else
+	else if (InFlattenMaterial.MetallicSamples.Num() == 1)
 	{
 		// Set Metallic to constant
-		float Metallic = InFlattenMaterial.MetallicSamples.Num() ? InFlattenMaterial.MetallicSamples[0].R : 0.0f;
+		float Metallic = *(float*)(&InFlattenMaterial.MetallicSamples[0].DWColor());
 		auto MetallicExpression = NewObject<UMaterialExpressionConstant>(Material);
 		MetallicExpression->R = Metallic;
 		MetallicExpression->MaterialExpressionEditorX = -400;
@@ -1086,7 +976,20 @@ UMaterial* FMaterialUtilities::CreateMaterial(const FFlattenMaterial& InFlattenM
 
 		MaterialNodeY+= MaterialNodeStepY;
 	}
+	else if (InFlattenMaterial.SpecularSamples.Num() == 1)
+	{
+		// Set Specular to constant
+		float Specular = *(float*)(&InFlattenMaterial.SpecularSamples[0].DWColor());
+		auto SpecularExpression = NewObject<UMaterialExpressionConstant>(Material);
+		SpecularExpression->R = Specular;
+		SpecularExpression->MaterialExpressionEditorX = -400;
+		SpecularExpression->MaterialExpressionEditorY = MaterialNodeY;
+		Material->Expressions.Add(SpecularExpression);
+		Material->Specular.Expression = SpecularExpression;
 
+		MaterialNodeY+= MaterialNodeStepY;
+	}
+	
 	// Roughness
 	if (InFlattenMaterial.RoughnessSamples.Num() > 1)
 	{
@@ -1105,10 +1008,10 @@ UMaterial* FMaterialUtilities::CreateMaterial(const FFlattenMaterial& InFlattenM
 
 		MaterialNodeY+= MaterialNodeStepY;
 	}
-	else
+	else if (InFlattenMaterial.RoughnessSamples.Num() == 1)
 	{
 		// Set Roughness to constant
-		float Roughness = InFlattenMaterial.RoughnessSamples.Num() ? InFlattenMaterial.RoughnessSamples[0].R : 0.8f;
+		float Roughness = *(float*)(&InFlattenMaterial.RoughnessSamples[0].DWColor());
 		auto RoughnessExpression = NewObject<UMaterialExpressionConstant>(Material);
 		RoughnessExpression->R = Roughness;
 		RoughnessExpression->MaterialExpressionEditorX = -400;
diff --git a/Engine/Source/Developer/MeshUtilities/Private/MeshUtilities.cpp b/Engine/Source/Developer/MeshUtilities/Private/MeshUtilities.cpp
--- a/Engine/Source/Developer/MeshUtilities/Private/MeshUtilities.cpp
+++ b/Engine/Source/Developer/MeshUtilities/Private/MeshUtilities.cpp
@@ -3648,6 +3648,24 @@ static bool NonOpaqueMaterialPredicate(UStaticMeshComponent* InMesh)
 	return false; 
 }
 
+static FIntPoint ConditionalImageResize(const FIntPoint& SrcSize, const FIntPoint& DesiredSize, TArray<FColor>& InOutImage, bool bLinearSpace)
+{
+	const int32 NumDesiredSamples = DesiredSize.X*DesiredSize.Y;
+	if (InOutImage.Num() && InOutImage.Num() != NumDesiredSamples)
+	{
+		check(InOutImage.Num() == SrcSize.X*SrcSize.Y);
+		TArray<FColor> OutImage;
+		if (NumDesiredSamples > 0)
+		{
+			FImageUtils::ImageResize(SrcSize.X, SrcSize.Y, InOutImage, DesiredSize.X, DesiredSize.Y, OutImage, bLinearSpace);
+		}
+		Exchange(InOutImage, OutImage);
+		return DesiredSize;
+	}
+
+	return SrcSize;
+}
+
 void FMeshUtilities::CreateProxyMesh( 
 	const TArray<AActor*>& SourceActors, 
 	const struct FMeshProxySettings& InProxySettings,
@@ -3725,13 +3743,11 @@ void FMeshUtilities::CreateProxyMesh(
 	// Convert materials into flatten materials
 	{
 		FFlattenMaterial FlattenMaterial;
-		FIntPoint TargetTextureSize = FIntPoint(InProxySettings.TextureWidth, InProxySettings.TextureHeight);
-
-		FlattenMaterial.DiffuseSize = TargetTextureSize;
-		FlattenMaterial.NormalSize = InProxySettings.bExportNormalMap ?	TargetTextureSize : FIntPoint::ZeroValue;
-		FlattenMaterial.MetallicSize = InProxySettings.bExportMetallicMap ? TargetTextureSize : FIntPoint::ZeroValue;
-		FlattenMaterial.RoughnessSize = InProxySettings.bExportRoughnessMap ? TargetTextureSize : FIntPoint::ZeroValue;
-		FlattenMaterial.SpecularSize = InProxySettings.bExportSpecularMap ? TargetTextureSize : FIntPoint::ZeroValue;
+		FlattenMaterial.DiffuseSize		= InProxySettings.Material.BaseColorMapSize;
+		FlattenMaterial.NormalSize		= InProxySettings.Material.bNormalMap ?	InProxySettings.Material.NormalMapSize	: FIntPoint::ZeroValue;
+		FlattenMaterial.MetallicSize	= InProxySettings.Material.bMetallicMap ? InProxySettings.Material.MetallicMapSize : FIntPoint::ZeroValue;
+		FlattenMaterial.RoughnessSize	= InProxySettings.Material.bRoughnessMap ? InProxySettings.Material.RoughnessMapSize : FIntPoint::ZeroValue;
+		FlattenMaterial.SpecularSize	= InProxySettings.Material.bSpecularMap ? InProxySettings.Material.SpecularMapSize : FIntPoint::ZeroValue;
 
 		for (UMaterialInterface* Material : StaticMeshMaterials)
 		{
@@ -3772,6 +3788,35 @@ void FMeshUtilities::CreateProxyMesh(
 		Vertex-= OutProxyLocation;
 	}
 	
+	{
+		// Resize merged textures into user requested size
+		ProxyFlattenMaterial.DiffuseSize	= ConditionalImageResize(ProxyFlattenMaterial.DiffuseSize, InProxySettings.Material.BaseColorMapSize, ProxyFlattenMaterial.DiffuseSamples, false);
+		ProxyFlattenMaterial.NormalSize		= ConditionalImageResize(ProxyFlattenMaterial.NormalSize, InProxySettings.Material.NormalMapSize, ProxyFlattenMaterial.NormalSamples, true);
+		ProxyFlattenMaterial.MetallicSize	= ConditionalImageResize(ProxyFlattenMaterial.MetallicSize, InProxySettings.Material.MetallicMapSize, ProxyFlattenMaterial.MetallicSamples, true);
+		ProxyFlattenMaterial.RoughnessSize	= ConditionalImageResize(ProxyFlattenMaterial.RoughnessSize, InProxySettings.Material.RoughnessMapSize, ProxyFlattenMaterial.RoughnessSamples, true);
+		ProxyFlattenMaterial.SpecularSize	= ConditionalImageResize(ProxyFlattenMaterial.SpecularSize, InProxySettings.Material.SpecularMapSize, ProxyFlattenMaterial.SpecularSamples, true);
+		
+		// Fill proxy material constants
+		if (ProxyFlattenMaterial.MetallicSamples.Num() == 0)
+		{
+			ProxyFlattenMaterial.MetallicSize = FIntPoint(1, 1);
+			ProxyFlattenMaterial.MetallicSamples.SetNum(1);
+			ProxyFlattenMaterial.MetallicSamples[0].DWColor() = *(uint32*)(&InProxySettings.Material.MetallicConstant);
+		}
+		if (ProxyFlattenMaterial.RoughnessSamples.Num() == 0)
+		{
+			ProxyFlattenMaterial.RoughnessSize = FIntPoint(1, 1);
+			ProxyFlattenMaterial.RoughnessSamples.SetNum(1);
+			ProxyFlattenMaterial.RoughnessSamples[0].DWColor() = *(uint32*)(&InProxySettings.Material.RoughnessConstant);
+		}
+		if (ProxyFlattenMaterial.SpecularSamples.Num() == 0)
+		{
+			ProxyFlattenMaterial.SpecularSize = FIntPoint(1, 1);
+			ProxyFlattenMaterial.SpecularSamples.SetNum(1);
+			ProxyFlattenMaterial.SpecularSamples[0].DWColor() = *(uint32*)(&InProxySettings.Material.SpecularConstant);
+		}
+	}
+
 	// Construct proxy material
 	UMaterial* ProxyMaterial = FMaterialUtilities::CreateMaterial(ProxyFlattenMaterial, InOuter, ProxyBasePackageName, RF_Public|RF_Standalone, OutAssetsToSync);
 	
diff --git a/Engine/Source/Developer/SimplygonMeshReduction/Private/SimplygonMeshReduction.cpp b/Engine/Source/Developer/SimplygonMeshReduction/Private/SimplygonMeshReduction.cpp
--- a/Engine/Source/Developer/SimplygonMeshReduction/Private/SimplygonMeshReduction.cpp
+++ b/Engine/Source/Developer/SimplygonMeshReduction/Private/SimplygonMeshReduction.cpp
@@ -515,8 +515,10 @@ public:
 		MappingSettings->SetGenerateMappingImage( true );
 		MappingSettings->SetGenerateTexCoords( true );
 		MappingSettings->SetGenerateTangents( false );
-		MappingSettings->SetWidth( InProxySettings.TextureWidth );
-		MappingSettings->SetHeight( InProxySettings.TextureHeight );
+
+		FIntPoint MappingImageSize = ComputeMappingImageSize(InProxySettings.Material);
+		MappingSettings->SetWidth(MappingImageSize.X);
+		MappingSettings->SetHeight(MappingImageSize.Y);
 
 		//Start remeshing the geometry
 		RemeshingProcessor->RemeshGeometry();
@@ -1698,6 +1700,16 @@ private:
 		// Specular
 		GetMaterialChannelData(SGMaterial, USER_MATERIAL_CHANNEL_SPECULAR, OutMaterial.SpecularSamples, OutMaterial.SpecularSize);
 	}
+
+	static FIntPoint ComputeMappingImageSize(const FMaterialSimplificationSettings& Settings)
+	{
+		FIntPoint ImageSize = Settings.BaseColorMapSize;
+		ImageSize = ImageSize.ComponentMax(Settings.NormalMapSize);
+		ImageSize = ImageSize.ComponentMax(Settings.MetallicMapSize);
+		ImageSize = ImageSize.ComponentMax(Settings.RoughnessMapSize);
+		ImageSize = ImageSize.ComponentMax(Settings.SpecularMapSize);
+		return ImageSize;
+	}
 };
 
 TScopedPointer<FSimplygonMeshReduction> GSimplygonMeshReduction;
diff --git a/Engine/Source/Editor/MergeActors/Private/MeshProxyTool/SMeshProxyDialog.cpp b/Engine/Source/Editor/MergeActors/Private/MeshProxyTool/SMeshProxyDialog.cpp
--- a/Engine/Source/Editor/MergeActors/Private/MeshProxyTool/SMeshProxyDialog.cpp
+++ b/Engine/Source/Editor/MergeActors/Private/MeshProxyTool/SMeshProxyDialog.cpp
@@ -512,49 +512,53 @@ void SMeshProxyDialog::SetTextureResolution(TSharedPtr<FString> NewSelection, ES
 {
 	int32 Resolution = 512;
 	TTypeFromString<int32>::FromString(Resolution, **NewSelection);
+	FIntPoint TextureSize(Resolution, Resolution);
 	
-	Tool->ProxySettings.TextureWidth	= Resolution;
-	Tool->ProxySettings.TextureHeight = Resolution;
+	Tool->ProxySettings.Material.BaseColorMapSize = TextureSize;
+	Tool->ProxySettings.Material.NormalMapSize = TextureSize;
+	Tool->ProxySettings.Material.MetallicMapSize = TextureSize;
+	Tool->ProxySettings.Material.RoughnessMapSize = TextureSize;
+	Tool->ProxySettings.Material.SpecularMapSize = TextureSize;
 }
 
 ECheckBoxState SMeshProxyDialog::GetExportNormalMap() const
 {
-	return Tool->ProxySettings.bExportNormalMap ? ECheckBoxState::Checked :  ECheckBoxState::Unchecked;
+	return Tool->ProxySettings.Material.bNormalMap ? ECheckBoxState::Checked :  ECheckBoxState::Unchecked;
 }
 
 void SMeshProxyDialog::SetExportNormalMap(ECheckBoxState NewValue)
 {
-	Tool->ProxySettings.bExportNormalMap = (NewValue == ECheckBoxState::Checked);
+	Tool->ProxySettings.Material.bNormalMap = (NewValue == ECheckBoxState::Checked);
 }
 
 ECheckBoxState SMeshProxyDialog::GetExportMetallicMap() const
 {
-	return Tool->ProxySettings.bExportMetallicMap ? ECheckBoxState::Checked :  ECheckBoxState::Unchecked;
+	return Tool->ProxySettings.Material.bMetallicMap ? ECheckBoxState::Checked :  ECheckBoxState::Unchecked;
 }
 
 void SMeshProxyDialog::SetExportMetallicMap(ECheckBoxState NewValue)
 {
-	Tool->ProxySettings.bExportMetallicMap = (NewValue == ECheckBoxState::Checked);
+	Tool->ProxySettings.Material.bMetallicMap = (NewValue == ECheckBoxState::Checked);
 }
 
 ECheckBoxState SMeshProxyDialog::GetExportRoughnessMap() const
 {
-	return Tool->ProxySettings.bExportRoughnessMap ? ECheckBoxState::Checked :  ECheckBoxState::Unchecked;
+	return Tool->ProxySettings.Material.bRoughnessMap ? ECheckBoxState::Checked :  ECheckBoxState::Unchecked;
 }
 
 void SMeshProxyDialog::SetExportRoughnessMap(ECheckBoxState NewValue)
 {
-	Tool->ProxySettings.bExportRoughnessMap = (NewValue == ECheckBoxState::Checked);
+	Tool->ProxySettings.Material.bRoughnessMap = (NewValue == ECheckBoxState::Checked);
 }
 
 ECheckBoxState SMeshProxyDialog::GetExportSpecularMap() const
 {
-	return Tool->ProxySettings.bExportSpecularMap ? ECheckBoxState::Checked :  ECheckBoxState::Unchecked;
+	return Tool->ProxySettings.Material.bSpecularMap ? ECheckBoxState::Checked :  ECheckBoxState::Unchecked;
 }
 
 void SMeshProxyDialog::SetExportSpecularMap(ECheckBoxState NewValue)
 {
-	Tool->ProxySettings.bExportSpecularMap = (NewValue == ECheckBoxState::Checked);
+	Tool->ProxySettings.Material.bSpecularMap = (NewValue == ECheckBoxState::Checked);
 }
 
 
diff --git a/Engine/Source/Editor/WorldBrowser/Private/Tiles/WorldTileCollectionModel.cpp b/Engine/Source/Editor/WorldBrowser/Private/Tiles/WorldTileCollectionModel.cpp
--- a/Engine/Source/Editor/WorldBrowser/Private/Tiles/WorldTileCollectionModel.cpp
+++ b/Engine/Source/Editor/WorldBrowser/Private/Tiles/WorldTileCollectionModel.cpp
@@ -1958,12 +1958,7 @@ bool FWorldTileCollectionModel::GenerateLODLevels(FLevelModelList InLevelList, i
 
 			FMeshProxySettings ProxySettings;
 			ProxySettings.ScreenSize = ProxySettings.ScreenSize*(SimplificationDetails.DetailsPercentage/100.f);
-			ProxySettings.TextureWidth = 1024; // TODO: Expose texture size
-			ProxySettings.TextureHeight = 1024;
-			ProxySettings.bExportNormalMap = SimplificationDetails.bGenerateMeshNormalMap;
-			ProxySettings.bExportMetallicMap = SimplificationDetails.bGenerateMeshMetallicMap;
-			ProxySettings.bExportRoughnessMap = SimplificationDetails.bGenerateMeshRoughnessMap;
-			ProxySettings.bExportSpecularMap = SimplificationDetails.bGenerateMeshSpecularMap;
+			ProxySettings.Material = SimplificationDetails.StaticMeshMaterial;
 
 			TArray<UObject*> OutAssets;
 			FVector OutProxyLocation;
@@ -2033,14 +2028,35 @@ bool FWorldTileCollectionModel::GenerateLODLevels(FLevelModelList InLevelList, i
 			}
 								
 			// This is texture resolution for a landscape mesh, probably needs to be calculated using landscape size
-			const FIntPoint LandscapeTextureSize(1024, 1024);
-			LandscapeFlattenMaterial.DiffuseSize	= LandscapeTextureSize;
-			LandscapeFlattenMaterial.NormalSize		= SimplificationDetails.bGenerateLandscapeNormalMap ? LandscapeTextureSize : FIntPoint::ZeroValue;
-			LandscapeFlattenMaterial.MetallicSize	= SimplificationDetails.bGenerateLandscapeMetallicMap ? LandscapeTextureSize : FIntPoint::ZeroValue;
-			LandscapeFlattenMaterial.RoughnessSize	= SimplificationDetails.bGenerateLandscapeRoughnessMap ? LandscapeTextureSize : FIntPoint::ZeroValue;
-			LandscapeFlattenMaterial.SpecularSize	= SimplificationDetails.bGenerateLandscapeSpecularMap ? LandscapeTextureSize : FIntPoint::ZeroValue;
+			LandscapeFlattenMaterial.DiffuseSize	= SimplificationDetails.LandscapeMaterial.BaseColorMapSize;
+			LandscapeFlattenMaterial.NormalSize		= SimplificationDetails.LandscapeMaterial.bNormalMap ?  SimplificationDetails.LandscapeMaterial.NormalMapSize : FIntPoint::ZeroValue;
+			LandscapeFlattenMaterial.MetallicSize	= SimplificationDetails.LandscapeMaterial.bMetallicMap ? SimplificationDetails.LandscapeMaterial.MetallicMapSize : FIntPoint::ZeroValue;
+			LandscapeFlattenMaterial.RoughnessSize	= SimplificationDetails.LandscapeMaterial.bRoughnessMap ? SimplificationDetails.LandscapeMaterial.RoughnessMapSize : FIntPoint::ZeroValue;
+			LandscapeFlattenMaterial.SpecularSize	= SimplificationDetails.LandscapeMaterial.bSpecularMap ? SimplificationDetails.LandscapeMaterial.SpecularMapSize : FIntPoint::ZeroValue;
 			
 			FMaterialUtilities::ExportLandscapeMaterial(Landscape, PrimitivesToHide, LandscapeFlattenMaterial);
+
+			// Fill landscape material constants
+			{
+				if (LandscapeFlattenMaterial.MetallicSamples.Num() == 0)
+				{
+					LandscapeFlattenMaterial.MetallicSize = FIntPoint(1, 1);
+					LandscapeFlattenMaterial.MetallicSamples.SetNum(1);
+					LandscapeFlattenMaterial.MetallicSamples[0].DWColor() = *(uint32*)(&SimplificationDetails.LandscapeMaterial.MetallicConstant);
+				}
+				if (LandscapeFlattenMaterial.RoughnessSamples.Num() == 0)
+				{
+					LandscapeFlattenMaterial.RoughnessSize = FIntPoint(1, 1);
+					LandscapeFlattenMaterial.RoughnessSamples.SetNum(1);
+					LandscapeFlattenMaterial.RoughnessSamples[0].DWColor() = *(uint32*)(&SimplificationDetails.LandscapeMaterial.RoughnessConstant);
+				}
+				if (LandscapeFlattenMaterial.SpecularSamples.Num() == 0)
+				{
+					LandscapeFlattenMaterial.SpecularSize = FIntPoint(1, 1);
+					LandscapeFlattenMaterial.SpecularSamples.SetNum(1);
+					LandscapeFlattenMaterial.SpecularSamples[0].DWColor() = *(uint32*)(&SimplificationDetails.LandscapeMaterial.SpecularConstant);
+				}
+			}
 		
 			if (SimplificationDetails.bBakeGrassToLandscape)
 			{
diff --git a/Engine/Source/Runtime/Engine/Classes/Engine/EngineTypes.h b/Engine/Source/Runtime/Engine/Classes/Engine/EngineTypes.h
--- a/Engine/Source/Runtime/Engine/Classes/Engine/EngineTypes.h
+++ b/Engine/Source/Runtime/Engine/Classes/Engine/EngineTypes.h
@@ -2229,6 +2229,91 @@ struct FMeshBuildSettings
 	}
 };
 
+USTRUCT()
+struct FMaterialSimplificationSettings
+{
+	GENERATED_USTRUCT_BODY()
+
+	// Size of generated BaseColor map
+	UPROPERTY(Category=Material, EditAnywhere)
+	FIntPoint BaseColorMapSize;
+
+	// Whether to generate normal map
+	UPROPERTY()
+	bool bNormalMap;
+	
+	// Size of generated specular map
+	UPROPERTY(Category=Material, EditAnywhere, meta=(editcondition = "bNormalMap"))
+	FIntPoint NormalMapSize;
+
+	// Metallic constant
+	UPROPERTY(Category=Material, EditAnywhere, meta=(ClampMin = "0", ClampMax = "1", UIMin = "0", UIMax = "1"))
+	float MetallicConstant;
+
+	// Whether to generate metallic map
+	UPROPERTY()
+	bool bMetallicMap;
+
+	// Size of generated metallic map
+	UPROPERTY(Category=Material, EditAnywhere, meta=(editcondition = "bMetallicMap"))
+	FIntPoint MetallicMapSize;
+
+	// Roughness constant
+	UPROPERTY(Category=Material, EditAnywhere, meta=(ClampMin = "0", ClampMax = "1", UIMin = "0", UIMax = "1"))
+	float RoughnessConstant;
+
+	// Whether to generate roughness map
+	UPROPERTY()
+	bool bRoughnessMap;
+
+	// Size of generated roughness map
+	UPROPERTY(Category=Material, EditAnywhere, meta=(editcondition = "bRoughnessMap"))
+	FIntPoint RoughnessMapSize;
+		
+	// Specular constant
+	UPROPERTY(Category=Material, EditAnywhere, meta=(ClampMin = "0", ClampMax = "1", UIMin = "0", UIMax = "1"))
+	float SpecularConstant;
+
+	// Whether to generate specular map
+	UPROPERTY()
+	bool bSpecularMap;
+
+	// Size of generated specular map
+	UPROPERTY(Category=Material, EditAnywhere, meta=(editcondition = "bSpecularMap"))
+	FIntPoint SpecularMapSize;
+
+	FMaterialSimplificationSettings()
+		: BaseColorMapSize(1024, 1024)
+		, bNormalMap(true)
+		, NormalMapSize(1024, 1024)
+		, MetallicConstant(0.0f)
+		, bMetallicMap(false)
+		, MetallicMapSize(1024, 1024)
+		, RoughnessConstant(0.5f)
+		, bRoughnessMap(false)
+		, RoughnessMapSize(1024, 1024)
+		, SpecularConstant(0.5f)
+		, bSpecularMap(false)
+		, SpecularMapSize(1024, 1024)
+	{
+	}
+	
+	bool operator == (const FMaterialSimplificationSettings& Other) const
+	{
+		return BaseColorMapSize == Other.BaseColorMapSize
+			&& bNormalMap == Other.bNormalMap
+			&& NormalMapSize == Other.NormalMapSize
+			&& MetallicConstant == Other.MetallicConstant
+			&& bMetallicMap == Other.bMetallicMap
+			&& MetallicMapSize == Other.MetallicMapSize
+			&& RoughnessConstant == Other.RoughnessConstant
+			&& bRoughnessMap == Other.bRoughnessMap
+			&& RoughnessMapSize == Other.RoughnessMapSize
+			&& SpecularConstant == Other.SpecularConstant
+			&& bSpecularMap == Other.bSpecularMap
+			&& SpecularMapSize == Other.SpecularMapSize;
+	}
+};
 
 USTRUCT()
 struct FMeshProxySettings
@@ -2238,27 +2323,26 @@ struct FMeshProxySettings
 	UPROPERTY(EditAnywhere, Category=ProxySettings)
 	int32 ScreenSize;
 
-	/** Texture resolution used for the proxy mesh*/
-	UPROPERTY(EditAnywhere, Category=ProxySettings)
-	int32 TextureWidth;
+	/** Material simplification */
 	UPROPERTY(EditAnywhere, Category=ProxySettings)
-	int32 TextureHeight;
+	FMaterialSimplificationSettings Material;
 
-	/** Whether to create normal map for the proxy mesh*/
-	UPROPERTY(EditAnywhere, Category=ProxySettings)
-	bool bExportNormalMap;
+	UPROPERTY()
+	int32 TextureWidth_DEPRECATED;
+	UPROPERTY()
+	int32 TextureHeight_DEPRECATED;
 
-	/** Whether to create metallic map for the proxy mesh*/
-	UPROPERTY(EditAnywhere, Category=ProxySettings)
-	bool bExportMetallicMap;
+	UPROPERTY()
+	bool bExportNormalMap_DEPRECATED;
 
-	/** Whether to create roughness map for the proxy mesh*/
-	UPROPERTY(EditAnywhere, Category=ProxySettings)
-	bool bExportRoughnessMap;
+	UPROPERTY()
+	bool bExportMetallicMap_DEPRECATED;
 
-	/** Whether to create specular map for the proxy mesh*/
-	UPROPERTY(EditAnywhere, Category=ProxySettings)
-	bool bExportSpecularMap;
+	UPROPERTY()
+	bool bExportRoughnessMap_DEPRECATED;
+
+	UPROPERTY()
+	bool bExportSpecularMap_DEPRECATED;
 
 	/** Should Simplygon recalculate normals for the proxy mesh? */
 	UPROPERTY(EditAnywhere, Category=ProxySettings)
@@ -2286,12 +2370,12 @@ struct FMeshProxySettings
 	/** Default settings. */
 	FMeshProxySettings()
 		: ScreenSize(300)
-		, TextureWidth(512)
-		, TextureHeight(512)
-		, bExportNormalMap(true)
-		, bExportMetallicMap(false)
-		, bExportRoughnessMap(false)
-		, bExportSpecularMap(false)
+		, TextureWidth_DEPRECATED(512)
+		, TextureHeight_DEPRECATED(512)
+		, bExportNormalMap_DEPRECATED(true)
+		, bExportMetallicMap_DEPRECATED(false)
+		, bExportRoughnessMap_DEPRECATED(false)
+		, bExportSpecularMap_DEPRECATED(false)
 		, bRecalculateNormals(true)
 		, HardAngleThreshold(80.0f)
 		, MergeDistance(4)
@@ -2305,12 +2389,7 @@ struct FMeshProxySettings
 	bool operator==(const FMeshProxySettings& Other) const
 	{
 		return ScreenSize == Other.ScreenSize
-			&& TextureWidth == Other.TextureWidth
-			&& TextureHeight == Other.TextureHeight
-			&& bExportNormalMap == Other.bExportNormalMap
-			&& bExportMetallicMap == Other.bExportMetallicMap
-			&& bExportRoughnessMap == Other.bExportRoughnessMap
-			&& bExportSpecularMap == Other.bExportSpecularMap
+			&& Material == Other.Material
 			&& bRecalculateNormals == Other.bRecalculateNormals
 			&& HardAngleThreshold == Other.HardAngleThreshold
 			&& MergeDistance == Other.MergeDistance;
@@ -2321,6 +2400,9 @@ struct FMeshProxySettings
 	{
 		return !(*this == Other);
 	}
+
+	/** Handles deprecated properties */
+	void PostLoadDeprecated();
 };
 
 /**
@@ -2340,7 +2422,7 @@ struct FMeshMergingSettings
 	int32 TargetLightMapUVChannel;
 
 	/** Target lightmap resolution */
-UPROPERTY(EditAnywhere, Category=FMeshMergingSettings)
+	UPROPERTY(EditAnywhere, Category=FMeshMergingSettings)
 	int32 TargetLightMapResolution;
 		
 	/** Whether we should import vertex colors into merged mesh */
diff --git a/Engine/Source/Runtime/Engine/Classes/Engine/Level.h b/Engine/Source/Runtime/Engine/Classes/Engine/Level.h
--- a/Engine/Source/Runtime/Engine/Classes/Engine/Level.h
+++ b/Engine/Source/Runtime/Engine/Classes/Engine/Level.h
@@ -230,71 +230,77 @@ private:
 	friend class FLightmassProcessor;
 };
 
+#if WITH_EDITOR
+
 USTRUCT()
 struct ENGINE_API FLevelSimplificationDetails
 {
 	GENERATED_USTRUCT_BODY()
 
-	// Whether to create separate packages for each generated asset. All in map package otherwise
+	/** Whether to create separate packages for each generated asset. All in map package otherwise */
 	UPROPERTY(Category=General, EditAnywhere)
 	bool bCreatePackagePerAsset;
 
-	// Percentage of details for static mesh proxy
+	/** Percentage of details for static mesh proxy */
 	UPROPERTY(Category=StaticMesh, EditAnywhere, meta=(DisplayName="Static Mesh Details Percentage", ClampMin = "0", ClampMax = "100", UIMin = "0", UIMax = "100"))	
 	float DetailsPercentage;
 
-	// Whether to generate normal map for static mesh proxy
-	UPROPERTY(Category=StaticMesh, EditAnywhere, meta=(DisplayName="Static Mesh Normal Map"))
-	bool bGenerateMeshNormalMap;
-	
-	// Whether to generate metallic map for static mesh proxy
-	UPROPERTY(Category=StaticMesh, EditAnywhere, meta=(DisplayName="Static Mesh Metallic Map"))
-	bool bGenerateMeshMetallicMap;
-
-	// Whether to generate roughness map for static mesh proxy
-	UPROPERTY(Category=StaticMesh, EditAnywhere, meta=(DisplayName="Static Mesh Roughness Map"))
-	bool bGenerateMeshRoughnessMap;
-	
-	// Whether to generate specular map for static mesh proxy
-	UPROPERTY(Category=StaticMesh, EditAnywhere, meta=(DisplayName="Static Mesh Specular Map"))
-	bool bGenerateMeshSpecularMap;
+	/** Static mesh material simplification */
+	UPROPERTY(Category=StaticMesh, EditAnywhere)
+	FMaterialSimplificationSettings StaticMeshMaterial;
 
 	UPROPERTY()
 	bool bOverrideLandscapeExportLOD;
 
-	// Landscape LOD to use for static mesh generation, when not specified 'Max LODLevel' from landscape actor will be used
+	/** Landscape LOD to use for static mesh generation, when not specified 'Max LODLevel' from landscape actor will be used */
 	UPROPERTY(Category=Landscape, EditAnywhere, meta=(ClampMin = "0", ClampMax = "7", UIMin = "0", UIMax = "7", editcondition = "bOverrideLandscapeExportLOD"))
 	int32 LandscapeExportLOD;
 
-	// Whether to generate normal map for landscape static mesh
-	UPROPERTY(Category=Landscape, EditAnywhere, meta=(DisplayName="Landscape Normal Map"))
-	bool bGenerateLandscapeNormalMap;
-
-	// Whether to generate metallic map for landscape static mesh
-	UPROPERTY(Category=Landscape, EditAnywhere, meta=(DisplayName="Landscape Metallic Map"))
-	bool bGenerateLandscapeMetallicMap;
+	/** Landscape material simplification */
+	UPROPERTY(Category=Landscape, EditAnywhere)
+	FMaterialSimplificationSettings LandscapeMaterial;
 
-	// Whether to generate roughness map for landscape static mesh
-	UPROPERTY(Category=Landscape, EditAnywhere, meta=(DisplayName="Landscape Roughness Map"))
-	bool bGenerateLandscapeRoughnessMap;
-	
-	// Whether to generate specular map for landscape static mesh
-	UPROPERTY(Category=Landscape, EditAnywhere, meta=(DisplayName="Landscape Specular Map"))
-	bool bGenerateLandscapeSpecularMap;
-	
-	// Whether to bake foliage into landscape static mesh texture
+	/** Whether to bake foliage into landscape static mesh texture */
 	UPROPERTY(Category=Landscape, EditAnywhere)
 	bool bBakeFoliageToLandscape;
 
-	// Whether to bake grass into landscape static mesh texture
+	/** Whether to bake grass into landscape static mesh texture */
 	UPROPERTY(Category=Landscape, EditAnywhere)
 	bool bBakeGrassToLandscape;
 
-	FLevelSimplificationDetails();
+	UPROPERTY()
+	bool bGenerateMeshNormalMap_DEPRECATED;
+	
+	UPROPERTY()
+	bool bGenerateMeshMetallicMap_DEPRECATED;
+
+	UPROPERTY()
+	bool bGenerateMeshRoughnessMap_DEPRECATED;
+	
+	UPROPERTY()
+	bool bGenerateMeshSpecularMap_DEPRECATED;
+	
+	UPROPERTY()
+	bool bGenerateLandscapeNormalMap_DEPRECATED;
+
+	UPROPERTY()
+	bool bGenerateLandscapeMetallicMap_DEPRECATED;
 
+	UPROPERTY()
+	bool bGenerateLandscapeRoughnessMap_DEPRECATED;
+	
+	UPROPERTY()
+	bool bGenerateLandscapeSpecularMap_DEPRECATED;
+
+	/** Handles deprecated properties */
+	void PostLoadDeprecated();
+	
+	FLevelSimplificationDetails();
 	bool operator == (const FLevelSimplificationDetails& Other) const;
 };
 
+#endif //WITH_EDITOR
+
 //
 // The level object.  Contains the level's actor list, BSP information, and brush list.
 // Every Level has a World as its Outer and can be used as the PersistentLevel, however,
diff --git a/Engine/Source/Runtime/Engine/Private/EngineTypes.cpp b/Engine/Source/Runtime/Engine/Private/EngineTypes.cpp
--- a/Engine/Source/Runtime/Engine/Private/EngineTypes.cpp
+++ b/Engine/Source/Runtime/Engine/Private/EngineTypes.cpp
@@ -3,6 +3,45 @@
 #include "EnginePrivate.h"
 #include "PhysicalMaterials/PhysicalMaterial.h"
 
+
+void FMeshProxySettings::PostLoadDeprecated()
+{
+	FMeshProxySettings DefaultObject;
+	
+	if (TextureWidth_DEPRECATED != DefaultObject.TextureHeight_DEPRECATED)
+	{
+		Material.BaseColorMapSize.X = TextureWidth_DEPRECATED;
+		Material.NormalMapSize.X = TextureWidth_DEPRECATED;
+		Material.MetallicMapSize.X = TextureWidth_DEPRECATED;
+		Material.RoughnessMapSize.X = TextureWidth_DEPRECATED;
+		Material.SpecularMapSize.X = TextureWidth_DEPRECATED;
+	}
+	if (TextureHeight_DEPRECATED != DefaultObject.TextureHeight_DEPRECATED)
+	{
+		Material.BaseColorMapSize.Y = TextureHeight_DEPRECATED;
+		Material.NormalMapSize.Y = TextureHeight_DEPRECATED;
+		Material.MetallicMapSize.Y = TextureHeight_DEPRECATED;
+		Material.RoughnessMapSize.Y = TextureHeight_DEPRECATED;
+		Material.SpecularMapSize.Y = TextureHeight_DEPRECATED;
+	}
+	if (bExportNormalMap_DEPRECATED != DefaultObject.bExportNormalMap_DEPRECATED)
+	{
+		Material.bNormalMap = bExportNormalMap_DEPRECATED;
+	}
+	if (bExportMetallicMap_DEPRECATED != DefaultObject.bExportMetallicMap_DEPRECATED)
+	{
+		Material.bMetallicMap = bExportMetallicMap_DEPRECATED;
+	}
+	if (bExportRoughnessMap_DEPRECATED != DefaultObject.bExportRoughnessMap_DEPRECATED)
+	{
+		Material.bRoughnessMap = bExportRoughnessMap_DEPRECATED;
+	}
+	if (bExportSpecularMap_DEPRECATED != DefaultObject.bExportSpecularMap_DEPRECATED)
+	{
+		Material.bSpecularMap = bExportSpecularMap_DEPRECATED;
+	}
+}
+
 UEngineBaseTypes::UEngineBaseTypes(const FObjectInitializer& ObjectInitializer)
 	: Super(ObjectInitializer)
 {
diff --git a/Engine/Source/Runtime/Engine/Private/ImageUtils.cpp b/Engine/Source/Runtime/Engine/Private/ImageUtils.cpp
--- a/Engine/Source/Runtime/Engine/Private/ImageUtils.cpp
+++ b/Engine/Source/Runtime/Engine/Private/ImageUtils.cpp
@@ -22,7 +22,7 @@ DEFINE_LOG_CATEGORY_STATIC(LogImageUtils, Log, All);
  */
 void FImageUtils::ImageResize(int32 SrcWidth, int32 SrcHeight, const TArray<FColor> &SrcData, int32 DstWidth, int32 DstHeight, TArray<FColor> &DstData, bool bLinearSpace )
 {
-	DstData.Empty();
+	DstData.Empty(DstWidth*DstHeight);
 	DstData.AddZeroed(DstWidth*DstHeight);
 
 	float SrcX = 0;
diff --git a/Engine/Source/Runtime/Engine/Private/Level.cpp b/Engine/Source/Runtime/Engine/Private/Level.cpp
--- a/Engine/Source/Runtime/Engine/Private/Level.cpp
+++ b/Engine/Source/Runtime/Engine/Private/Level.cpp
@@ -163,39 +163,69 @@ FArchive& operator<<( FArchive& Ar, FPrecomputedVolumeDistanceField& D )
 
 FLevelSimplificationDetails::FLevelSimplificationDetails()
  : bCreatePackagePerAsset(true)
- , DetailsPercentage(70.f)
- , bGenerateMeshNormalMap(true)
- , bGenerateMeshMetallicMap(false)
- , bGenerateMeshRoughnessMap(false)
- , bGenerateMeshSpecularMap(false)
+ , DetailsPercentage(70.0f)
  , bOverrideLandscapeExportLOD(false)
  , LandscapeExportLOD(7)
- , bGenerateLandscapeNormalMap(true)
- , bGenerateLandscapeMetallicMap(false)
- , bGenerateLandscapeRoughnessMap(false)
- , bGenerateLandscapeSpecularMap(false)
  , bBakeFoliageToLandscape(false)
  , bBakeGrassToLandscape(false)
+ , bGenerateMeshNormalMap_DEPRECATED(true)
+ , bGenerateMeshMetallicMap_DEPRECATED(false)
+ , bGenerateMeshRoughnessMap_DEPRECATED(false)
+ , bGenerateMeshSpecularMap_DEPRECATED(false)
+ , bGenerateLandscapeNormalMap_DEPRECATED(true)
+ , bGenerateLandscapeMetallicMap_DEPRECATED(false)
+ , bGenerateLandscapeRoughnessMap_DEPRECATED(false)
+ , bGenerateLandscapeSpecularMap_DEPRECATED(false)
 {
 }
 
 bool FLevelSimplificationDetails::operator == (const FLevelSimplificationDetails& Other) const
 {
-	return
-		bCreatePackagePerAsset == Other.bCreatePackagePerAsset &&
-		DetailsPercentage == Other.DetailsPercentage &&
-		bGenerateMeshNormalMap == Other.bGenerateMeshNormalMap &&
-		bGenerateMeshMetallicMap == Other.bGenerateMeshMetallicMap &&
-		bGenerateMeshRoughnessMap == Other.bGenerateMeshRoughnessMap &&
-		bGenerateMeshSpecularMap == Other.bGenerateMeshSpecularMap &&
-		bOverrideLandscapeExportLOD == Other.bOverrideLandscapeExportLOD &&
-		LandscapeExportLOD == Other.LandscapeExportLOD &&
-		bGenerateLandscapeNormalMap == Other.bGenerateLandscapeNormalMap &&
-		bGenerateLandscapeMetallicMap == Other.bGenerateLandscapeMetallicMap &&
-		bGenerateLandscapeRoughnessMap == Other.bGenerateLandscapeRoughnessMap &&
-		bGenerateLandscapeSpecularMap == Other.bGenerateLandscapeSpecularMap &&
-		bBakeFoliageToLandscape == Other.bBakeFoliageToLandscape &&
-		bBakeGrassToLandscape == Other.bBakeGrassToLandscape;
+	return bCreatePackagePerAsset == Other.bCreatePackagePerAsset
+		&& DetailsPercentage == Other.DetailsPercentage
+		&& StaticMeshMaterial == Other.StaticMeshMaterial
+		&& bOverrideLandscapeExportLOD == Other.bOverrideLandscapeExportLOD
+		&& LandscapeExportLOD == Other.LandscapeExportLOD
+		&& LandscapeMaterial == Other.LandscapeMaterial
+		&& bBakeFoliageToLandscape == Other.bBakeFoliageToLandscape
+		&& bBakeGrassToLandscape == Other.bBakeGrassToLandscape;
+}
+
+void FLevelSimplificationDetails::PostLoadDeprecated()
+{
+	FLevelSimplificationDetails DefaultObject;
+	if (bGenerateMeshNormalMap_DEPRECATED != DefaultObject.bGenerateMeshNormalMap_DEPRECATED)
+	{
+		StaticMeshMaterial.bNormalMap = bGenerateMeshNormalMap_DEPRECATED;
+	}
+	if (bGenerateMeshMetallicMap_DEPRECATED != DefaultObject.bGenerateMeshMetallicMap_DEPRECATED)
+	{
+		StaticMeshMaterial.bMetallicMap = bGenerateMeshMetallicMap_DEPRECATED;
+	}
+	if (bGenerateMeshRoughnessMap_DEPRECATED != DefaultObject.bGenerateMeshRoughnessMap_DEPRECATED)
+	{
+		StaticMeshMaterial.bRoughnessMap = bGenerateMeshRoughnessMap_DEPRECATED;
+	}
+	if (bGenerateMeshSpecularMap_DEPRECATED != DefaultObject.bGenerateMeshSpecularMap_DEPRECATED)
+	{
+		StaticMeshMaterial.bSpecularMap = bGenerateMeshSpecularMap_DEPRECATED;
+	}
+	if (bGenerateLandscapeNormalMap_DEPRECATED != DefaultObject.bGenerateLandscapeNormalMap_DEPRECATED)
+	{
+		LandscapeMaterial.bNormalMap = bGenerateLandscapeNormalMap_DEPRECATED;
+	}
+	if (bGenerateLandscapeMetallicMap_DEPRECATED != DefaultObject.bGenerateLandscapeMetallicMap_DEPRECATED)
+	{
+		LandscapeMaterial.bMetallicMap = bGenerateLandscapeMetallicMap_DEPRECATED;
+	}
+	if (bGenerateLandscapeRoughnessMap_DEPRECATED != DefaultObject.bGenerateLandscapeRoughnessMap_DEPRECATED)
+	{
+		LandscapeMaterial.bRoughnessMap = bGenerateLandscapeRoughnessMap_DEPRECATED;
+	}
+	if (bGenerateLandscapeSpecularMap_DEPRECATED != DefaultObject.bGenerateLandscapeSpecularMap_DEPRECATED)
+	{
+		LandscapeMaterial.bSpecularMap = bGenerateLandscapeSpecularMap_DEPRECATED;
+	}
 }
 
 TMap<FName, TWeakObjectPtr<UWorld> > ULevel::StreamedLevelsOwningWorld;
@@ -547,6 +577,13 @@ void ULevel::PostLoad()
 			LevelScriptBlueprint->Rename(*OuterWorld->GetName(), LevelScriptBlueprint->GetOuter(), REN_DoNotDirty | REN_DontCreateRedirectors | REN_ForceNoResetLoaders | REN_NonTransactional | REN_SkipGeneratedClasses);
 		}
 	}
+
+	// Fixup deprecated stuff in levels simplification settings
+	for (int32 Index = 0; Index < ARRAY_COUNT(LevelSimplification); ++Index)
+	{
+		LevelSimplification[Index].PostLoadDeprecated();
+	}
+
 #endif
 }
 
