MeshMerging supports LightmassReplace material node during material flattening

diff --git a/Engine/Source/Developer/MaterialUtilities/Private/MaterialUtilities.cpp b/Engine/Source/Developer/MaterialUtilities/Private/MaterialUtilities.cpp
--- a/Engine/Source/Developer/MaterialUtilities/Private/MaterialUtilities.cpp
+++ b/Engine/Source/Developer/MaterialUtilities/Private/MaterialUtilities.cpp
@@ -89,6 +89,8 @@ struct FExportMaterialCompiler : public FProxyMaterialCompiler
 	{
 		return Compiler->Constant4(1.0f,1.0f,1.0f,1.0f);
 	}
+	
+	virtual int32 LightmassReplace(int32 Realtime, int32 Lightmass) override { return Lightmass; }
 };
 
 
