Fixed issues with normal map during material flattening

diff --git a/Engine/Source/Developer/MaterialUtilities/Private/MaterialUtilities.cpp b/Engine/Source/Developer/MaterialUtilities/Private/MaterialUtilities.cpp
--- a/Engine/Source/Developer/MaterialUtilities/Private/MaterialUtilities.cpp
+++ b/Engine/Source/Developer/MaterialUtilities/Private/MaterialUtilities.cpp
@@ -675,7 +675,7 @@ bool FMaterialUtilities::ExportMaterial(UWorld* InWorld, UMaterialInterface* InM
 		RenderTargetNormal->ClearColor = FLinearColor(0.0f, 0.0f, 0.0f, 0.0f);
 		RenderTargetNormal->InitCustomFormat(
 			OutFlattenMaterial.NormalSize.X,
-			OutFlattenMaterial.NormalSize.Y, PF_FloatRGB, true);
+			OutFlattenMaterial.NormalSize.Y, PF_B8G8R8A8, true);
 
 		//
 		OutFlattenMaterial.NormalSamples.Empty(OutFlattenMaterial.NormalSize.X * OutFlattenMaterial.NormalSize.Y);
