Copying //UE4/Dev-Framework to Dev-Main (//UE4/Dev-Main)

diff --git a/Engine/Source/Developer/LogVisualizer/Private/VisualLoggerDatabase.cpp b/Engine/Source/Developer/LogVisualizer/Private/VisualLoggerDatabase.cpp
--- a/Engine/Source/Developer/LogVisualizer/Private/VisualLoggerDatabase.cpp
+++ b/Engine/Source/Developer/LogVisualizer/Private/VisualLoggerDatabase.cpp
@@ -160,6 +160,11 @@ void FVisualLoggerDatabase::AddItem(const FVisualLogDevice::FVisualLogEntryItem&
 	FVisualLoggerGraphsDatabase::Get().AddItem(NewItem);
 }
 
+bool FVisualLoggerDatabase::ContainsRowByName(FName InName)
+{
+	return RowNameToIndex.Contains(InName);
+}
+
 FVisualLoggerDBRow& FVisualLoggerDatabase::GetRowByName(FName InName)
 {
 	const bool bContainsRow = RowNameToIndex.Contains(InName);
diff --git a/Engine/Source/Developer/LogVisualizer/Private/VisualLoggerDatabase.h b/Engine/Source/Developer/LogVisualizer/Private/VisualLoggerDatabase.h
--- a/Engine/Source/Developer/LogVisualizer/Private/VisualLoggerDatabase.h
+++ b/Engine/Source/Developer/LogVisualizer/Private/VisualLoggerDatabase.h
@@ -80,6 +80,7 @@ struct FVisualLoggerDatabase
 	FConstRowIterator GetConstRowIterator() const { return Rows.CreateConstIterator(); }
 	FRowIterator GetRowIterator() { return Rows.CreateIterator(); }
 	
+	bool ContainsRowByName(FName InName);
 	FVisualLoggerDBRow& GetRowByName(FName InName);
 
 	void SelectRow(FName InName, bool bDeselectOtherNodes = false);
diff --git a/Engine/Source/Developer/LogVisualizer/Private/VisualLoggerHUD.h b/Engine/Source/Developer/LogVisualizer/Private/VisualLoggerHUD.h
--- a/Engine/Source/Developer/LogVisualizer/Private/VisualLoggerHUD.h
+++ b/Engine/Source/Developer/LogVisualizer/Private/VisualLoggerHUD.h
@@ -4,7 +4,7 @@
 #include "Engine/DebugCameraHUD.h"
 #include "VisualLoggerHUD.generated.h"
 
-UCLASS()
+UCLASS(hidedropdown)
 class AVisualLoggerHUD : public ADebugCameraHUD
 {
 	GENERATED_UCLASS_BODY()
diff --git a/Engine/Source/Developer/LogVisualizer/Public/LogVisualizerSettings.h b/Engine/Source/Developer/LogVisualizer/Public/LogVisualizerSettings.h
--- a/Engine/Source/Developer/LogVisualizer/Public/LogVisualizerSettings.h
+++ b/Engine/Source/Developer/LogVisualizer/Public/LogVisualizerSettings.h
@@ -24,7 +24,9 @@ struct FCategoryFilter
 	int32 LogVerbosity;
 
 	UPROPERTY(config)
-	bool Enabled;
+	uint32 Enabled : 1;
+
+	uint32 bIsInUse : 1;
 };
 
 USTRUCT()
@@ -58,6 +60,7 @@ struct FVisualLoggerFilters : public FVisualLoggerFiltersData
 	static void Shutdown();
 
 	void Reset();
+	void InitWith(const FVisualLoggerFiltersData& NewFiltersData);
 
 	bool MatchCategoryFilters(FString String, ELogVerbosity::Type Verbosity = ELogVerbosity::All);
 
diff --git a/Engine/Source/Developer/MaterialUtilities/Private/MaterialUtilities.cpp b/Engine/Source/Developer/MaterialUtilities/Private/MaterialUtilities.cpp
--- a/Engine/Source/Developer/MaterialUtilities/Private/MaterialUtilities.cpp
+++ b/Engine/Source/Developer/MaterialUtilities/Private/MaterialUtilities.cpp
@@ -71,7 +71,7 @@ struct FExportMaterialCompiler : public FProxyMaterialCompiler
 
 	virtual int32 ObjectWorldPosition() override
 	{
-		return Compiler->Constant3(0.0f, 0.0f, 0.0f);
+		return Compiler->ObjectWorldPosition();
 	}
 
 	virtual int32 DistanceCullFade() override
@@ -81,7 +81,7 @@ struct FExportMaterialCompiler : public FProxyMaterialCompiler
 
 	virtual int32 ActorWorldPosition() override
 	{
-		return Compiler->Constant3(0.0f, 0.0f, 0.0f);
+		return Compiler->ActorWorldPosition();
 	}
 
 	virtual int32 ParticleRelativeTime() override
@@ -93,6 +93,12 @@ struct FExportMaterialCompiler : public FProxyMaterialCompiler
 	{
 		return Compiler->Constant(1.0f);
 	}
+	
+	virtual int32 PixelNormalWS() override
+	{
+		// Current returning vertex normal since pixel normal will contain incorrect data (normal calculated from uv data used as vertex positions to render out the material)
+		return Compiler->VertexNormal();
+	}
 
 	virtual int32 ParticleRandom() override
 	{
@@ -121,7 +127,7 @@ struct FExportMaterialCompiler : public FProxyMaterialCompiler
 
 	virtual int32 ObjectBounds() override
 	{
-		return Compiler->Constant3(0, 0, 0);
+		return Compiler->ObjectBounds();
 	}
 
 	virtual int32 CameraVector() override
@@ -131,7 +137,7 @@ struct FExportMaterialCompiler : public FProxyMaterialCompiler
 	
 	virtual int32 ReflectionAboutCustomWorldNormal(int32 CustomWorldNormal, int32 bNormalizeCustomWorldNormal) override
 	{
-		return Compiler->Constant3(0.0f, 0.0f, -1.0f);
+		return Compiler->ReflectionAboutCustomWorldNormal(CustomWorldNormal, bNormalizeCustomWorldNormal);
 	}
 
 	virtual int32 VertexColor() override
@@ -141,12 +147,12 @@ struct FExportMaterialCompiler : public FProxyMaterialCompiler
 
 	virtual int32 LightVector() override
 	{
-		return Compiler->Constant3(1.0f, 0.0f, 0.0f);
+		return Compiler->LightVector();
 	}
 
 	virtual int32 ReflectionVector() override
 	{
-		return Compiler->Constant3(0.0f, 0.0f, -1.0f);
+		return Compiler->ReflectionVector();
 	}
 
 	virtual int32 AtmosphericFogColor(int32 WorldPosition) override
@@ -257,6 +263,8 @@ public:
 		case MP_Roughness: ResourceId.Usage = EMaterialShaderMapUsage::MaterialExportRoughness; break;
 		case MP_AmbientOcclusion: ResourceId.Usage = EMaterialShaderMapUsage::MaterialExportAO; break;
 		case MP_EmissiveColor: ResourceId.Usage = EMaterialShaderMapUsage::MaterialExportEmissive; break;
+		case MP_Opacity: ResourceId.Usage = EMaterialShaderMapUsage::MaterialExportOpacity; break;
+		case MP_SubsurfaceColor: ResourceId.Usage = EMaterialShaderMapUsage::MaterialExportSubSurfaceColor; break;
 		default:
 			ensureMsgf(false, TEXT("ExportMaterial has no usage for property %i.  Will likely reuse the normal rendering shader and crash later with a parameter mismatch"), (int32)InPropertyToCompile);
 			break;
@@ -647,7 +655,8 @@ bool FMaterialUtilities::ExportMaterialProperty(UWorld* InWorld, UMaterialInterf
 	const bool bForceGamma = (InMaterialProperty == MP_Normal) || (InMaterialProperty == MP_OpacityMask) || (InMaterialProperty == MP_Opacity);	
 
 	FIntPoint MaxSize = MaterialProxy->FindMaxTextureSize(InMaterial);
-	return RenderMaterialPropertyToTexture(MaterialData, InMaterialProperty, bForceGamma, PF_B8G8R8A8, MaxSize, OutBMP);
+	FIntPoint OutSize = MaxSize;
+	return RenderMaterialPropertyToTexture(MaterialData, InMaterialProperty, bForceGamma, PF_B8G8R8A8, MaxSize, OutSize, OutBMP);
 }
 
 bool FMaterialUtilities::ExportMaterialProperty(UWorld* InWorld, UMaterialInterface* InMaterial, EMaterialProperty InMaterialProperty, FIntPoint& OutSize, TArray<FColor>& OutBMP)
@@ -663,7 +672,7 @@ bool FMaterialUtilities::ExportMaterialProperty(UWorld* InWorld, UMaterialInterf
 	FMaterialMergeData MaterialData(InMaterial, nullptr, nullptr, 0, DummyBounds, EmptyTexCoords);
 	const bool bForceGamma = (InMaterialProperty == MP_Normal) || (InMaterialProperty == MP_OpacityMask) || (InMaterialProperty == MP_Opacity);
 	OutSize = MaterialProxy->FindMaxTextureSize(InMaterial);
-	return RenderMaterialPropertyToTexture(MaterialData, InMaterialProperty, bForceGamma, PF_B8G8R8A8, OutSize, OutBMP);
+	return RenderMaterialPropertyToTexture(MaterialData, InMaterialProperty, bForceGamma, PF_B8G8R8A8, OutSize, OutSize, OutBMP);
 }
 
 bool FMaterialUtilities::ExportMaterialProperty(UMaterialInterface* InMaterial, EMaterialProperty InMaterialProperty, TArray<FColor>& OutBMP, FIntPoint& OutSize)
@@ -679,7 +688,7 @@ bool FMaterialUtilities::ExportMaterialProperty(UMaterialInterface* InMaterial,
 	FMaterialMergeData MaterialData(InMaterial, nullptr, nullptr, 0, DummyBounds, EmptyTexCoords);
 	const bool bForceGamma = (InMaterialProperty == MP_Normal) || (InMaterialProperty == MP_OpacityMask) || (InMaterialProperty == MP_Opacity);
 	OutSize = MaterialProxy->FindMaxTextureSize(InMaterial);
-	return RenderMaterialPropertyToTexture(MaterialData, InMaterialProperty, bForceGamma, PF_B8G8R8A8, OutSize, OutBMP);
+	return RenderMaterialPropertyToTexture(MaterialData, InMaterialProperty, bForceGamma, PF_B8G8R8A8, OutSize, OutSize, OutBMP);
 }
 
 bool FMaterialUtilities::ExportMaterialProperty(UMaterialInterface* InMaterial, EMaterialProperty InMaterialProperty, FIntPoint InSize, TArray<FColor>& OutBMP)
@@ -694,7 +703,8 @@ bool FMaterialUtilities::ExportMaterialProperty(UMaterialInterface* InMaterial,
 	TArray<FVector2D> EmptyTexCoords;
 	FMaterialMergeData MaterialData(InMaterial, nullptr, nullptr, 0, DummyBounds, EmptyTexCoords);
 	const bool bForceGamma = (InMaterialProperty == MP_Normal) || (InMaterialProperty == MP_OpacityMask) || (InMaterialProperty == MP_Opacity);
-	return RenderMaterialPropertyToTexture(MaterialData, InMaterialProperty, bForceGamma, PF_B8G8R8A8, InSize, OutBMP);
+	FIntPoint OutSize;
+	return RenderMaterialPropertyToTexture(MaterialData, InMaterialProperty, bForceGamma, PF_B8G8R8A8, InSize, OutSize, OutBMP);
 }
 
 bool FMaterialUtilities::ExportMaterial(UWorld* InWorld, UMaterialInterface* InMaterial, FFlattenMaterial& OutFlattenMaterial)
@@ -853,9 +863,9 @@ UMaterial* FMaterialUtilities::CreateMaterial(const FFlattenMaterial& InFlattenM
 	else if (InFlattenMaterial.DiffuseSamples.Num() == 1)
 	{
 		// Set Roughness to constant
-		FColor BaseColor = InFlattenMaterial.DiffuseSamples[0];
+		FLinearColor BaseColor = FLinearColor(InFlattenMaterial.DiffuseSamples[0]);
 		auto BaseColorExpression = NewObject<UMaterialExpressionConstant4Vector>(Material);
-		BaseColorExpression->Constant = BaseColor.ReinterpretAsLinear();
+		BaseColorExpression->Constant = BaseColor;
 		BaseColorExpression->MaterialExpressionEditorX = -400;
 		BaseColorExpression->MaterialExpressionEditorY = MaterialNodeY;
 		Material->Expressions.Add(BaseColorExpression);
@@ -875,15 +885,27 @@ UMaterial* FMaterialUtilities::CreateMaterial(const FFlattenMaterial& InFlattenM
 
 	// Check for same texture sizes
 	bool bSameTextureSize = true;	
-	bSameTextureSize &= bHasMetallic ? (InFlattenMaterial.DiffuseSamples.Num() == InFlattenMaterial.MetallicSamples.Num()) : true;
-	bSameTextureSize &= bHasRoughness ? (InFlattenMaterial.DiffuseSamples.Num() == InFlattenMaterial.RoughnessSamples.Num()) : true;
-	bSameTextureSize &= bHasSpecular ? (InFlattenMaterial.DiffuseSamples.Num() == InFlattenMaterial.SpecularSamples.Num()) : true;
+
+	int32 SampleCount = 0;
+	FIntPoint MergedSize(0,0);
+	SampleCount = (bHasMetallic && SampleCount == 0) ? InFlattenMaterial.MetallicSamples.Num() : SampleCount;
+	MergedSize = (bHasMetallic && MergedSize.X == 0) ? InFlattenMaterial.MetallicSize : MergedSize;
+
+	SampleCount = (bHasRoughness && SampleCount == 0) ? InFlattenMaterial.RoughnessSamples.Num() : SampleCount;
+	MergedSize = (bHasRoughness && MergedSize.X == 0) ? InFlattenMaterial.RoughnessSize : MergedSize;
+
+	SampleCount = (bHasSpecular && SampleCount == 0) ? InFlattenMaterial.SpecularSamples.Num() : SampleCount;
+	MergedSize = (bHasSpecular && MergedSize.X == 0) ? InFlattenMaterial.SpecularSize : MergedSize;
+
+	bSameTextureSize &= bHasMetallic ? (SampleCount == InFlattenMaterial.MetallicSamples.Num()) : true;
+	bSameTextureSize &= bHasRoughness ? (SampleCount == InFlattenMaterial.RoughnessSamples.Num()) : true;
+	bSameTextureSize &= bHasSpecular ? (SampleCount == InFlattenMaterial.SpecularSamples.Num()) : true;
+
 	// Merge values into one texture if more than one material property exists
 	if (BakedMaterialPropertyCount > 1 && bSameTextureSize)
 	{
 		// Metallic = R, Roughness = G, Specular = B
 		TArray<FColor> MergedSamples;
-		const int32 SampleCount = InFlattenMaterial.DiffuseSamples.Num();
 		MergedSamples.AddZeroed(SampleCount);
 		
 		if (bHasMetallic) for (int32 SampleIndex = 0; SampleIndex < SampleCount; ++SampleIndex)
@@ -902,13 +924,13 @@ UMaterial* FMaterialUtilities::CreateMaterial(const FFlattenMaterial& InFlattenM
 		}
 
 		const FString AssetName = TEXT("T_") + AssetBaseName + TEXT("_MRS");
-		const bool bSRGB = false;
-		UTexture2D* Texture = CreateTexture(InOuter, AssetBasePath + AssetName, InFlattenMaterial.DiffuseSize, MergedSamples, TC_Default, TEXTUREGROUP_World, Flags, bSRGB);
+		const bool bSRGB = true;
+		UTexture2D* Texture = CreateTexture(InOuter, AssetBasePath + AssetName, MergedSize, MergedSamples, TC_Default, TEXTUREGROUP_World, Flags, bSRGB);
 		OutGeneratedAssets.Add(Texture);
 
 		auto MergedExpression = NewObject<UMaterialExpressionTextureSample>(Material);
 		MergedExpression->Texture = Texture;
-		MergedExpression->SamplerType = EMaterialSamplerType::SAMPLERTYPE_LinearColor;
+		MergedExpression->SamplerType = EMaterialSamplerType::SAMPLERTYPE_Color;
 		MergedExpression->MaterialExpressionEditorX = -400;
 		MergedExpression->MaterialExpressionEditorY = MaterialNodeY;
 		Material->Expressions.Add(MergedExpression);
@@ -954,13 +976,13 @@ UMaterial* FMaterialUtilities::CreateMaterial(const FFlattenMaterial& InFlattenM
 		if (InFlattenMaterial.MetallicSamples.Num() > 1)
 		{
 			const FString AssetName = TEXT("T_") + AssetBaseName + TEXT("_M");
-			const bool bSRGB = false;
+			const bool bSRGB = true;
 			UTexture2D* Texture = CreateTexture(InOuter, AssetBasePath + AssetName, InFlattenMaterial.MetallicSize, InFlattenMaterial.MetallicSamples, TC_Grayscale, TEXTUREGROUP_World, Flags, bSRGB);
 			OutGeneratedAssets.Add(Texture);
 
 			auto MetallicExpression = NewObject<UMaterialExpressionTextureSample>(Material);
 			MetallicExpression->Texture = Texture;
-			MetallicExpression->SamplerType = EMaterialSamplerType::SAMPLERTYPE_LinearGrayscale;
+			MetallicExpression->SamplerType = EMaterialSamplerType::SAMPLERTYPE_Grayscale;
 			MetallicExpression->MaterialExpressionEditorX = -400;
 			MetallicExpression->MaterialExpressionEditorY = MaterialNodeY;
 			Material->Expressions.Add(MetallicExpression);
@@ -973,13 +995,13 @@ UMaterial* FMaterialUtilities::CreateMaterial(const FFlattenMaterial& InFlattenM
 		if (InFlattenMaterial.SpecularSamples.Num() > 1)
 		{
 			const FString AssetName = TEXT("T_") + AssetBaseName + TEXT("_S");
-			const bool bSRGB = false;
+			const bool bSRGB = true;
 			UTexture2D* Texture = CreateTexture(InOuter, AssetBasePath + AssetName, InFlattenMaterial.SpecularSize, InFlattenMaterial.SpecularSamples, TC_Grayscale, TEXTUREGROUP_World, Flags, bSRGB);
 			OutGeneratedAssets.Add(Texture);
 
 			auto SpecularExpression = NewObject<UMaterialExpressionTextureSample>(Material);
 			SpecularExpression->Texture = Texture;
-			SpecularExpression->SamplerType = EMaterialSamplerType::SAMPLERTYPE_LinearGrayscale;
+			SpecularExpression->SamplerType = EMaterialSamplerType::SAMPLERTYPE_Grayscale;
 			SpecularExpression->MaterialExpressionEditorX = -400;
 			SpecularExpression->MaterialExpressionEditorY = MaterialNodeY;
 			Material->Expressions.Add(SpecularExpression);
@@ -992,13 +1014,13 @@ UMaterial* FMaterialUtilities::CreateMaterial(const FFlattenMaterial& InFlattenM
 		if (InFlattenMaterial.RoughnessSamples.Num() > 1)
 		{
 			const FString AssetName = TEXT("T_") + AssetBaseName + TEXT("_R");
-			const bool bSRGB = false;
+			const bool bSRGB = true;
 			UTexture2D* Texture = CreateTexture(InOuter, AssetBasePath + AssetName, InFlattenMaterial.RoughnessSize, InFlattenMaterial.RoughnessSamples, TC_Grayscale, TEXTUREGROUP_World, Flags, bSRGB);
 			OutGeneratedAssets.Add(Texture);
 
 			auto RoughnessExpression = NewObject<UMaterialExpressionTextureSample>(Material);
 			RoughnessExpression->Texture = Texture;
-			RoughnessExpression->SamplerType = EMaterialSamplerType::SAMPLERTYPE_LinearGrayscale;
+			RoughnessExpression->SamplerType = EMaterialSamplerType::SAMPLERTYPE_Grayscale;
 			RoughnessExpression->MaterialExpressionEditorX = -400;
 			RoughnessExpression->MaterialExpressionEditorY = MaterialNodeY;
 			Material->Expressions.Add(RoughnessExpression);
@@ -1010,10 +1032,9 @@ UMaterial* FMaterialUtilities::CreateMaterial(const FFlattenMaterial& InFlattenM
 
 	if (InFlattenMaterial.MetallicSamples.Num() == 1)
 	{
-		// Set Metallic to constant
-		float Metallic = *(float*)(&InFlattenMaterial.MetallicSamples[0].DWColor());
+		FLinearColor MetallicColor = FLinearColor(InFlattenMaterial.MetallicSamples[0]);
 		auto MetallicExpression = NewObject<UMaterialExpressionConstant>(Material);
-		MetallicExpression->R = Metallic;
+		MetallicExpression->R = MetallicColor.R;
 		MetallicExpression->MaterialExpressionEditorX = -400;
 		MetallicExpression->MaterialExpressionEditorY = MaterialNodeY;
 		Material->Expressions.Add(MetallicExpression);
@@ -1025,9 +1046,9 @@ UMaterial* FMaterialUtilities::CreateMaterial(const FFlattenMaterial& InFlattenM
 	if (InFlattenMaterial.SpecularSamples.Num() == 1)
 	{
 		// Set Specular to constant
-		float Specular = *(float*)(&InFlattenMaterial.SpecularSamples[0].DWColor());
+		FLinearColor SpecularColor = FLinearColor(InFlattenMaterial.SpecularSamples[0]);
 		auto SpecularExpression = NewObject<UMaterialExpressionConstant>(Material);
-		SpecularExpression->R = Specular;
+		SpecularExpression->R = SpecularColor.R;
 		SpecularExpression->MaterialExpressionEditorX = -400;
 		SpecularExpression->MaterialExpressionEditorY = MaterialNodeY;
 		Material->Expressions.Add(SpecularExpression);
@@ -1039,9 +1060,9 @@ UMaterial* FMaterialUtilities::CreateMaterial(const FFlattenMaterial& InFlattenM
 	if (InFlattenMaterial.RoughnessSamples.Num() == 1)
 	{
 		// Set Roughness to constant
-		float Roughness = *(float*)(&InFlattenMaterial.RoughnessSamples[0].DWColor());
+		FLinearColor RoughnessColor = FLinearColor(InFlattenMaterial.RoughnessSamples[0]);
 		auto RoughnessExpression = NewObject<UMaterialExpressionConstant>(Material);
-		RoughnessExpression->R = Roughness;
+		RoughnessExpression->R = RoughnessColor.R;
 		RoughnessExpression->MaterialExpressionEditorX = -400;
 		RoughnessExpression->MaterialExpressionEditorY = MaterialNodeY;
 		Material->Expressions.Add(RoughnessExpression);
@@ -1073,12 +1094,11 @@ UMaterial* FMaterialUtilities::CreateMaterial(const FFlattenMaterial& InFlattenM
 	{
 		// Set Emissive to constant
 		FColor EmissiveColor = InFlattenMaterial.EmissiveSamples[0];
-
 		// Don't have to deal with black emissive color
 		if (EmissiveColor != FColor(0, 0, 0, 255))
 		{
 			auto EmissiveColorExpression = NewObject<UMaterialExpressionConstant4Vector>(Material);
-			EmissiveColorExpression->Constant = (EmissiveColor.ReinterpretAsLinear() * InFlattenMaterial.EmissiveScale);
+			EmissiveColorExpression->Constant = EmissiveColor.ReinterpretAsLinear() * InFlattenMaterial.EmissiveScale;
 			EmissiveColorExpression->MaterialExpressionEditorX = -400;
 			EmissiveColorExpression->MaterialExpressionEditorY = MaterialNodeY;
 			Material->Expressions.Add(EmissiveColorExpression);
@@ -1090,14 +1110,14 @@ UMaterial* FMaterialUtilities::CreateMaterial(const FFlattenMaterial& InFlattenM
 	else if (InFlattenMaterial.EmissiveSamples.Num() > 1)
 	{
 		const FString AssetName = TEXT("T_") + AssetBaseName + TEXT("_E");
-		const bool bSRGB = false;
-		UTexture2D* Texture = CreateTexture(InOuter, AssetBasePath + AssetName, InFlattenMaterial.EmissiveSize, InFlattenMaterial.EmissiveSamples, TC_Default, TEXTUREGROUP_WorldNormalMap, Flags, bSRGB);
+		const bool bSRGB = true;
+		UTexture2D* Texture = CreateTexture(InOuter, AssetBasePath + AssetName, InFlattenMaterial.EmissiveSize, InFlattenMaterial.EmissiveSamples, TC_Default, TEXTUREGROUP_World, Flags, bSRGB);
 		OutGeneratedAssets.Add(Texture);
 
 		//Assign emissive color to the material
 		UMaterialExpressionTextureSample* EmissiveColorExpression = NewObject<UMaterialExpressionTextureSample>(Material);
 		EmissiveColorExpression->Texture = Texture;
-		EmissiveColorExpression->SamplerType = EMaterialSamplerType::SAMPLERTYPE_LinearColor;
+		EmissiveColorExpression->SamplerType = EMaterialSamplerType::SAMPLERTYPE_Color;
 		EmissiveColorExpression->MaterialExpressionEditorX = -400;
 		EmissiveColorExpression->MaterialExpressionEditorY = MaterialNodeY;
 		Material->Expressions.Add(EmissiveColorExpression);
@@ -1116,9 +1136,9 @@ UMaterial* FMaterialUtilities::CreateMaterial(const FFlattenMaterial& InFlattenM
 	if (InFlattenMaterial.OpacitySamples.Num() == 1)
 	{
 		// Set Opacity to constant
-		float Opacity = *(float*)(&InFlattenMaterial.OpacitySamples[0].DWColor());
+		FLinearColor Opacity = FLinearColor(InFlattenMaterial.OpacitySamples[0]);
 		auto OpacityExpression = NewObject<UMaterialExpressionConstant>(Material);
-		OpacityExpression->R = Opacity;
+		OpacityExpression->R = Opacity.R;
 		OpacityExpression->MaterialExpressionEditorX = -400;
 		OpacityExpression->MaterialExpressionEditorY = MaterialNodeY;
 		Material->Expressions.Add(OpacityExpression);
@@ -1129,20 +1149,63 @@ UMaterial* FMaterialUtilities::CreateMaterial(const FFlattenMaterial& InFlattenM
 	else if (InFlattenMaterial.OpacitySamples.Num() > 1)
 	{
 		const FString AssetName = TEXT("T_") + AssetBaseName + TEXT("_O");
-		const bool bSRGB = false;
-		UTexture2D* Texture = CreateTexture(InOuter, AssetBasePath + AssetName, InFlattenMaterial.OpacitySize, InFlattenMaterial.OpacitySamples, TC_Grayscale, TEXTUREGROUP_WorldNormalMap, Flags, bSRGB);
+		const bool bSRGB = true;
+		UTexture2D* Texture = CreateTexture(InOuter, AssetBasePath + AssetName, InFlattenMaterial.OpacitySize, InFlattenMaterial.OpacitySamples, TC_Grayscale, TEXTUREGROUP_World, Flags, bSRGB);
 		OutGeneratedAssets.Add(Texture);
 
 		//Assign opacity to the material
 		UMaterialExpressionTextureSample* OpacityExpression = NewObject<UMaterialExpressionTextureSample>(Material);
 		OpacityExpression->Texture = Texture;
-		OpacityExpression->SamplerType = EMaterialSamplerType::SAMPLERTYPE_LinearGrayscale;
+		OpacityExpression->SamplerType = EMaterialSamplerType::SAMPLERTYPE_Grayscale;
 		OpacityExpression->MaterialExpressionEditorX = -400;
 		OpacityExpression->MaterialExpressionEditorY = MaterialNodeY;
 		Material->Expressions.Add(OpacityExpression);
 		Material->Opacity.Expression = OpacityExpression;
 		MaterialNodeY += MaterialNodeStepY;
 	}
+
+	if (InFlattenMaterial.SubSurfaceSamples.Num() == 1)
+	{
+		// Set Emissive to constant
+		FColor SubSurfaceColor = InFlattenMaterial.SubSurfaceSamples[0];
+
+		// Don't have to deal with black sub surface color
+		if (SubSurfaceColor != FColor(0, 0, 0, 255))
+		{
+			auto SubSurfaceColorExpression = NewObject<UMaterialExpressionConstant4Vector>(Material);
+			SubSurfaceColorExpression->Constant = (SubSurfaceColor.ReinterpretAsLinear());
+			SubSurfaceColorExpression->MaterialExpressionEditorX = -400;
+			SubSurfaceColorExpression->MaterialExpressionEditorY = MaterialNodeY;
+			Material->Expressions.Add(SubSurfaceColorExpression);
+			Material->SubsurfaceColor.Expression = SubSurfaceColorExpression;
+
+			MaterialNodeY += MaterialNodeStepY;
+		}
+
+		Material->SetShadingModel(MSM_Subsurface);
+
+	}
+	else if (InFlattenMaterial.SubSurfaceSamples.Num() > 1)
+	{
+		const FString AssetName = TEXT("T_") + AssetBaseName + TEXT("_SSC");
+		const bool bSRGB = true;
+		UTexture2D* Texture = CreateTexture(InOuter, AssetBasePath + AssetName, InFlattenMaterial.SubSurfaceSize, InFlattenMaterial.SubSurfaceSamples, TC_Default, TEXTUREGROUP_World, Flags, bSRGB);
+		OutGeneratedAssets.Add(Texture);
+
+		//Assign emissive color to the material
+		UMaterialExpressionTextureSample* SubSurfaceColorExpression = NewObject<UMaterialExpressionTextureSample>(Material);
+		SubSurfaceColorExpression->Texture = Texture;
+		SubSurfaceColorExpression->SamplerType = EMaterialSamplerType::SAMPLERTYPE_Color;
+		SubSurfaceColorExpression->MaterialExpressionEditorX = -400;
+		SubSurfaceColorExpression->MaterialExpressionEditorY = MaterialNodeY;
+		Material->Expressions.Add(SubSurfaceColorExpression);
+
+		Material->SubsurfaceColor.Expression = SubSurfaceColorExpression;
+		MaterialNodeY += MaterialNodeStepY;
+
+		Material->SetShadingModel(MSM_Subsurface);
+	}
+
 							
 	Material->PostEditChange();
 	return Material;
@@ -1220,6 +1283,53 @@ FFlattenMaterial FMaterialUtilities::CreateFlattenMaterialWithSettings(const FMa
 	// Create new material.
 	FFlattenMaterial Material;
 
+	FIntPoint MaximumSize = InMaterialLODSettings.TextureSize;	
+	// If the user is manually overriding the texture size, make sure we have the max texture size to render with
+	if (InMaterialLODSettings.TextureSizingType == TextureSizingType_UseManualOverrideTextureSize)
+	{
+		MaximumSize = (MaximumSize.X < InMaterialLODSettings.DiffuseTextureSize.X) ? MaximumSize : InMaterialLODSettings.DiffuseTextureSize;
+		MaximumSize = (InMaterialLODSettings.bSpecularMap && MaximumSize.X < InMaterialLODSettings.SpecularTextureSize.X) ? MaximumSize : InMaterialLODSettings.SpecularTextureSize;
+		MaximumSize = (InMaterialLODSettings.bMetallicMap && MaximumSize.X < InMaterialLODSettings.MetallicTextureSize.X) ? MaximumSize : InMaterialLODSettings.MetallicTextureSize;
+		MaximumSize = (InMaterialLODSettings.bRoughnessMap && MaximumSize.X < InMaterialLODSettings.RoughnessTextureSize.X) ? MaximumSize : InMaterialLODSettings.RoughnessTextureSize;
+		MaximumSize = (InMaterialLODSettings.bNormalMap && MaximumSize.X < InMaterialLODSettings.NormalTextureSize.X) ? MaximumSize : InMaterialLODSettings.NormalTextureSize;
+		MaximumSize = (InMaterialLODSettings.bEmissiveMap && MaximumSize.X < InMaterialLODSettings.EmissiveTextureSize.X) ? MaximumSize : InMaterialLODSettings.EmissiveTextureSize;
+		MaximumSize = (InMaterialLODSettings.bOpacityMap && MaximumSize.X < InMaterialLODSettings.OpacityTextureSize.X) ? MaximumSize : InMaterialLODSettings.OpacityTextureSize;
+	}
+	
+	if (InMaterialLODSettings.TextureSizingType == TextureSizingType_UseManualOverrideTextureSize)
+	{
+		Material.RenderSize = MaximumSize;
+
+		Material.DiffuseSize = InMaterialLODSettings.DiffuseTextureSize;
+		Material.SpecularSize = (InMaterialLODSettings.bSpecularMap) ? InMaterialLODSettings.SpecularTextureSize : FIntPoint::ZeroValue;
+		Material.MetallicSize = (InMaterialLODSettings.bMetallicMap) ? InMaterialLODSettings.MetallicTextureSize : FIntPoint::ZeroValue;
+		Material.RoughnessSize = (InMaterialLODSettings.bRoughnessMap) ? InMaterialLODSettings.RoughnessTextureSize : FIntPoint::ZeroValue;
+		Material.NormalSize = (InMaterialLODSettings.bNormalMap) ? InMaterialLODSettings.NormalTextureSize : FIntPoint::ZeroValue;
+		Material.EmissiveSize = (InMaterialLODSettings.bEmissiveMap) ? InMaterialLODSettings.EmissiveTextureSize : FIntPoint::ZeroValue;
+		Material.OpacitySize = (InMaterialLODSettings.bOpacityMap) ? InMaterialLODSettings.OpacityTextureSize : FIntPoint::ZeroValue;
+	}
+	else if (InMaterialLODSettings.TextureSizingType == TextureSizingType_UseAutomaticBiasedSizes)
+	{
+		Material.RenderSize = InMaterialLODSettings.TextureSize;
+		
+		int NormalSizeX, DiffuseSizeX, PropertiesSizeX;
+		NormalSizeX = InMaterialLODSettings.TextureSize.X;
+		DiffuseSizeX = FMath::Max(InMaterialLODSettings.TextureSize.X >> 1, 32);
+		PropertiesSizeX = FMath::Max(InMaterialLODSettings.TextureSize.X >> 2, 16);
+
+		Material.DiffuseSize = FIntPoint(DiffuseSizeX, DiffuseSizeX);
+		Material.NormalSize = (InMaterialLODSettings.bNormalMap) ? FIntPoint(NormalSizeX, NormalSizeX) : FIntPoint::ZeroValue;
+
+		FIntPoint PropertiesSize = FIntPoint(PropertiesSizeX, PropertiesSizeX);
+		Material.SpecularSize = (InMaterialLODSettings.bSpecularMap) ? PropertiesSize: FIntPoint::ZeroValue;
+		Material.MetallicSize = (InMaterialLODSettings.bMetallicMap) ? PropertiesSize: FIntPoint::ZeroValue;
+		Material.RoughnessSize = (InMaterialLODSettings.bRoughnessMap) ? PropertiesSize : FIntPoint::ZeroValue;		
+		Material.EmissiveSize = (InMaterialLODSettings.bEmissiveMap) ? PropertiesSize : FIntPoint::ZeroValue;
+		Material.OpacitySize = (InMaterialLODSettings.bOpacityMap) ? PropertiesSize : FIntPoint::ZeroValue;
+	}
+	
+	Material.RenderSize = InMaterialLODSettings.TextureSize;
+	
 	Material.DiffuseSize = InMaterialLODSettings.TextureSize;
 	Material.SpecularSize = (InMaterialLODSettings.bSpecularMap) ? InMaterialLODSettings.TextureSize : FIntPoint::ZeroValue;
 	Material.MetallicSize = (InMaterialLODSettings.bMetallicMap) ? InMaterialLODSettings.TextureSize : FIntPoint::ZeroValue;
@@ -1231,9 +1341,9 @@ FFlattenMaterial FMaterialUtilities::CreateFlattenMaterialWithSettings(const FMa
 	return Material;
 }
 
-void FMaterialUtilities::AnalyzeMaterial(UMaterialInterface* InMaterial, const struct FMaterialProxySettings& InMaterialSettings, int32& OutNumTexCoords, bool& OutUseVertexColor)
+void FMaterialUtilities::AnalyzeMaterial(UMaterialInterface* InMaterial, const struct FMaterialProxySettings& InMaterialSettings, int32& OutNumTexCoords, bool& OutRequiresVertexData)
 {
-	OutUseVertexColor = false;
+	OutRequiresVertexData = false;
 	OutNumTexCoords = 0;
 
 	bool PropertyBeingBaked[MP_Normal + 1];	
@@ -1270,11 +1380,11 @@ void FMaterialUtilities::AnalyzeMaterial(UMaterialInterface* InMaterial, const s
 
 			// Analyze this material channel.
 			int32 NumTextureCoordinates = 0;
-			bool bUseVertexColor = false;
-			InMaterial->AnalyzeMaterialProperty(Property, NumTextureCoordinates, bUseVertexColor);
+			bool bUseVertexData = false;
+			InMaterial->AnalyzeMaterialProperty(Property, NumTextureCoordinates, bUseVertexData);
 			// Accumulate data.
 			OutNumTexCoords = FMath::Max(NumTextureCoordinates, OutNumTexCoords);
-			OutUseVertexColor |= bUseVertexColor;
+			OutRequiresVertexData |= bUseVertexData;
 		}
 	}
 }
@@ -1283,7 +1393,7 @@ void FMaterialUtilities::RemapUniqueMaterialIndices(const TArray<UMaterialInterf
 {
 	// Gather material properties
 	TMap<UMaterialInterface*, int32> MaterialNumTexCoords;
-	TMap<UMaterialInterface*, bool>  MaterialUseVertexColors;
+	TMap<UMaterialInterface*, bool>  MaterialUseVertexData;
 
 	for (int32 MaterialIndex = 0; MaterialIndex < InMaterials.Num(); MaterialIndex++)
 	{
@@ -1298,45 +1408,23 @@ void FMaterialUtilities::RemapUniqueMaterialIndices(const TArray<UMaterialInterf
 		{
 			// We are not baking vertex data at all, don't analyze materials.
 			MaterialNumTexCoords.Add(Material, 0);
-			MaterialUseVertexColors.Add(Material, false);
+			MaterialUseVertexData.Add(Material, false);
 			continue;
 		}
 		int32 NumTexCoords = 0;
-		bool bUseVertexColors = false;
-		FMaterialUtilities::AnalyzeMaterial(Material, InMaterialProxySettings, NumTexCoords, bUseVertexColors);
+		bool bUseVertexData = false;
+		FMaterialUtilities::AnalyzeMaterial(Material, InMaterialProxySettings, NumTexCoords, bUseVertexData);
 		MaterialNumTexCoords.Add(Material, NumTexCoords);
-		MaterialUseVertexColors.Add(Material, bUseVertexColors);
-	}
-
-	// Check which meshes has vertex colors.
-	TArray<bool> MeshHasVertexColors;
-	MeshHasVertexColors.AddZeroed(InMeshData.Num());
-	if (bBakeVertexData && bMergeMaterials)
-	{
-		for (int32 MeshIndex = 0; MeshIndex < InMeshData.Num(); MeshIndex++)
-		{
-			for (int32 LODIndex = 0; LODIndex < MAX_STATIC_MESH_LODS; ++LODIndex)
-			{
-				const FRawMesh Mesh = InMeshData[MeshIndex].MeshLODData[LODIndex].RawMesh;				
-				bool bHasVertexColors = false;
-				for (int32 WedgeIndex = 0; WedgeIndex < Mesh.WedgeColors.Num(); WedgeIndex++)
-				{
-					if (Mesh.WedgeColors[WedgeIndex] != FColor::White)
-					{
-						bHasVertexColors = true;
-						break;
-					}
-				}
-				MeshHasVertexColors[MeshIndex] |= bHasVertexColors;
-				
-			}
-		}
+		MaterialUseVertexData.Add(Material, bUseVertexData);
 	}
 
 	// Build list of mesh's material properties.
 	OutMeshShouldBakeVertexData.Empty();
 	OutMeshShouldBakeVertexData.AddZeroed(InMeshData.Num());
 
+	TArray<bool> MeshUsesVertexData;
+	MeshUsesVertexData.AddZeroed(InMeshData.Num());
+
 	for (int32 MeshIndex = 0; MeshIndex < InMeshData.Num(); MeshIndex++)
 	{
 		for (int32 LODIndex = 0; LODIndex < MAX_STATIC_MESH_LODS; ++LODIndex)
@@ -1345,19 +1433,17 @@ void FMaterialUtilities::RemapUniqueMaterialIndices(const TArray<UMaterialInterf
 			{
 				const TArray<int32>& MeshMaterialMap = InMaterialMap[FMeshIdAndLOD(MeshIndex, LODIndex)];
 				int32 NumTexCoords = 0;
-				bool bUseVertexColors = false;
+				bool bUseVertexData = false;
 				// Accumulate data of all materials.
 				for (int32 LocalMaterialIndex = 0; LocalMaterialIndex < MeshMaterialMap.Num(); LocalMaterialIndex++)
 				{
 					UMaterialInterface* Material = InMaterials[MeshMaterialMap[LocalMaterialIndex]];
 					NumTexCoords = FMath::Max(NumTexCoords, MaterialNumTexCoords[Material]);
-					bUseVertexColors |= MaterialUseVertexColors[Material];
+					bUseVertexData |= MaterialUseVertexData[Material];
 				}
-				// Take into account presence of vertex colors in mesh.
-				bUseVertexColors &= MeshHasVertexColors[MeshIndex];
 				// Store data.
-				MeshHasVertexColors[MeshIndex] = bUseVertexColors;
-				OutMeshShouldBakeVertexData[MeshIndex] = bUseVertexColors || (NumTexCoords >= 2);
+				MeshUsesVertexData[MeshIndex] = bUseVertexData;
+				OutMeshShouldBakeVertexData[MeshIndex] = bUseVertexData || (NumTexCoords >= 2);
 			}
 			else
 			{
@@ -1399,9 +1485,9 @@ void FMaterialUtilities::RemapUniqueMaterialIndices(const TArray<UMaterialInterf
 				TArray<int32>& NewMeshMaterialMap = OutMaterialMap.Add(FMeshIdAndLOD(MeshIndex, LODIndex));
 				UStaticMesh* StaticMesh = InMeshData[MeshIndex].MeshLODData[LODIndex].SourceStaticMesh;
 
-				if (!MeshHasVertexColors[MeshIndex])
+				if (!MeshUsesVertexData[MeshIndex])
 				{
-					// No vertex colors - could merge materials with other meshes.
+					// No vertex data needed - could merge materials with other meshes.
 					if (!OutMeshShouldBakeVertexData[MeshIndex])
 					{
 						// Set to 'nullptr' if don't need to bake vertex data to be able to merge materials with any meshes
@@ -1435,7 +1521,7 @@ void FMaterialUtilities::RemapUniqueMaterialIndices(const TArray<UMaterialInterf
 			else
 			{
 				break;
-			}			
+			}
 		}
 	}
 
@@ -1460,6 +1546,147 @@ void FMaterialUtilities::OptimizeFlattenMaterial(FFlattenMaterial& InFlattenMate
 	OptimizeSampleArray(InFlattenMaterial.EmissiveSamples, InFlattenMaterial.EmissiveSize);
 }
 
+void FMaterialUtilities::ResizeFlattenMaterial(FFlattenMaterial& InFlattenMaterial, const struct FMeshProxySettings& InMeshProxySettings)
+{
+	const FMaterialProxySettings& MaterialSettings = InMeshProxySettings.MaterialSettings;
+
+	if (MaterialSettings.TextureSizingType == TextureSizingType_UseAutomaticBiasedSizes)
+	{
+		int NormalSizeX, DiffuseSizeX, PropertiesSizeX;
+		NormalSizeX = MaterialSettings.TextureSize.X;
+		DiffuseSizeX = FMath::Max(MaterialSettings.TextureSize.X >> 1, 32);
+		PropertiesSizeX = FMath::Max(MaterialSettings.TextureSize.X >> 2, 16);
+
+		if (InFlattenMaterial.DiffuseSize.X != DiffuseSizeX)
+		{
+			TArray<FColor> NewSamples;
+			FImageUtils::ImageResize(InFlattenMaterial.DiffuseSize.X, InFlattenMaterial.DiffuseSize.Y, InFlattenMaterial.DiffuseSamples, DiffuseSizeX, DiffuseSizeX, NewSamples, false);
+			InFlattenMaterial.DiffuseSamples.Reset(NewSamples.Num());
+			InFlattenMaterial.DiffuseSamples.Append(NewSamples);
+			InFlattenMaterial.DiffuseSize = FIntPoint(DiffuseSizeX, DiffuseSizeX);
+		}
+
+		if (InFlattenMaterial.SpecularSamples.Num() && MaterialSettings.bSpecularMap && InFlattenMaterial.SpecularSize.X != PropertiesSizeX)
+		{
+			TArray<FColor> NewSamples;
+			FImageUtils::ImageResize(InFlattenMaterial.SpecularSize.X, InFlattenMaterial.SpecularSize.Y, InFlattenMaterial.SpecularSamples, PropertiesSizeX, PropertiesSizeX, NewSamples, false);
+			InFlattenMaterial.SpecularSamples.Reset(NewSamples.Num());
+			InFlattenMaterial.SpecularSamples.Append(NewSamples);
+			InFlattenMaterial.SpecularSize = FIntPoint(PropertiesSizeX, PropertiesSizeX);
+		}
+
+		if (InFlattenMaterial.MetallicSamples.Num() && MaterialSettings.bMetallicMap && InFlattenMaterial.MetallicSize.X != PropertiesSizeX)
+		{
+			TArray<FColor> NewSamples;
+			FImageUtils::ImageResize(InFlattenMaterial.MetallicSize.X, InFlattenMaterial.MetallicSize.Y, InFlattenMaterial.MetallicSamples, PropertiesSizeX, PropertiesSizeX, NewSamples, false);
+			InFlattenMaterial.MetallicSamples.Reset(NewSamples.Num());
+			InFlattenMaterial.MetallicSamples.Append(NewSamples);
+			InFlattenMaterial.MetallicSize = FIntPoint(PropertiesSizeX, PropertiesSizeX);
+		}
+
+		if (InFlattenMaterial.RoughnessSamples.Num() && MaterialSettings.bRoughnessMap && InFlattenMaterial.RoughnessSize.X != PropertiesSizeX)
+		{
+			TArray<FColor> NewSamples;
+			FImageUtils::ImageResize(InFlattenMaterial.RoughnessSize.X, InFlattenMaterial.RoughnessSize.Y, InFlattenMaterial.RoughnessSamples, PropertiesSizeX, PropertiesSizeX, NewSamples, false);
+			InFlattenMaterial.RoughnessSamples.Reset(NewSamples.Num());
+			InFlattenMaterial.RoughnessSamples.Append(NewSamples);
+			InFlattenMaterial.RoughnessSize = FIntPoint(PropertiesSizeX, PropertiesSizeX);
+		}
+
+		if (InFlattenMaterial.NormalSamples.Num() && MaterialSettings.bNormalMap && InFlattenMaterial.NormalSize.X != NormalSizeX)
+		{
+			TArray<FColor> NewSamples;
+			FImageUtils::ImageResize(InFlattenMaterial.NormalSize.X, InFlattenMaterial.NormalSize.Y, InFlattenMaterial.NormalSamples, NormalSizeX, NormalSizeX, NewSamples, true);
+			InFlattenMaterial.NormalSamples.Reset(NewSamples.Num());
+			InFlattenMaterial.NormalSamples.Append(NewSamples);
+			InFlattenMaterial.NormalSize = FIntPoint(NormalSizeX, NormalSizeX);
+		}
+
+		if (InFlattenMaterial.EmissiveSamples.Num() && MaterialSettings.bEmissiveMap && InFlattenMaterial.EmissiveSize.X != PropertiesSizeX)
+		{
+			TArray<FColor> NewSamples;
+			FImageUtils::ImageResize(InFlattenMaterial.EmissiveSize.X, InFlattenMaterial.EmissiveSize.Y, InFlattenMaterial.EmissiveSamples, PropertiesSizeX, PropertiesSizeX, NewSamples, false);
+			InFlattenMaterial.EmissiveSamples.Reset(NewSamples.Num());
+			InFlattenMaterial.EmissiveSamples.Append(NewSamples);
+			InFlattenMaterial.EmissiveSize = FIntPoint(PropertiesSizeX, PropertiesSizeX);
+		}
+
+		if (InFlattenMaterial.OpacitySamples.Num() && MaterialSettings.bOpacityMap && InFlattenMaterial.OpacitySize.X != PropertiesSizeX)
+		{
+			TArray<FColor> NewSamples;
+			FImageUtils::ImageResize(InFlattenMaterial.OpacitySize.X, InFlattenMaterial.OpacitySize.Y, InFlattenMaterial.OpacitySamples, PropertiesSizeX, PropertiesSizeX, NewSamples, false);
+			InFlattenMaterial.OpacitySamples.Reset(NewSamples.Num());
+			InFlattenMaterial.OpacitySamples.Append(NewSamples);
+			InFlattenMaterial.OpacitySize = FIntPoint(PropertiesSizeX, PropertiesSizeX);
+		}
+	}
+	else if (MaterialSettings.TextureSizingType == TextureSizingType_UseManualOverrideTextureSize)
+	{
+		if (InFlattenMaterial.DiffuseSize != MaterialSettings.DiffuseTextureSize)
+		{
+			TArray<FColor> NewSamples;
+			FImageUtils::ImageResize(InFlattenMaterial.DiffuseSize.X, InFlattenMaterial.DiffuseSize.Y, InFlattenMaterial.DiffuseSamples, MaterialSettings.DiffuseTextureSize.X, MaterialSettings.DiffuseTextureSize.Y, NewSamples, false);
+			InFlattenMaterial.DiffuseSamples.Reset(NewSamples.Num());
+			InFlattenMaterial.DiffuseSamples.Append(NewSamples);
+			InFlattenMaterial.DiffuseSize = MaterialSettings.DiffuseTextureSize;
+		}
+
+		if (InFlattenMaterial.SpecularSamples.Num() && MaterialSettings.bSpecularMap && InFlattenMaterial.SpecularSize != MaterialSettings.SpecularTextureSize)
+		{
+			TArray<FColor> NewSamples;
+			FImageUtils::ImageResize(InFlattenMaterial.SpecularSize.X, InFlattenMaterial.SpecularSize.Y, InFlattenMaterial.SpecularSamples, MaterialSettings.SpecularTextureSize.X, MaterialSettings.SpecularTextureSize.Y, NewSamples, false);
+			InFlattenMaterial.SpecularSamples.Reset(NewSamples.Num());
+			InFlattenMaterial.SpecularSamples.Append(NewSamples);
+			InFlattenMaterial.SpecularSize = MaterialSettings.SpecularTextureSize;
+		}
+
+		if (InFlattenMaterial.MetallicSamples.Num() && MaterialSettings.bMetallicMap && InFlattenMaterial.MetallicSize != MaterialSettings.MetallicTextureSize)
+		{
+			TArray<FColor> NewSamples;
+			FImageUtils::ImageResize(InFlattenMaterial.MetallicSize.X, InFlattenMaterial.MetallicSize.Y, InFlattenMaterial.MetallicSamples, MaterialSettings.MetallicTextureSize.X, MaterialSettings.MetallicTextureSize.Y, NewSamples, false);
+			InFlattenMaterial.MetallicSamples.Reset(NewSamples.Num());
+			InFlattenMaterial.MetallicSamples.Append(NewSamples);
+			InFlattenMaterial.MetallicSize = MaterialSettings.MetallicTextureSize;
+		}
+
+		if (InFlattenMaterial.RoughnessSamples.Num() && MaterialSettings.bRoughnessMap && InFlattenMaterial.RoughnessSize != MaterialSettings.RoughnessTextureSize)
+		{
+			TArray<FColor> NewSamples;
+			FImageUtils::ImageResize(InFlattenMaterial.RoughnessSize.X, InFlattenMaterial.RoughnessSize.Y, InFlattenMaterial.RoughnessSamples, MaterialSettings.RoughnessTextureSize.X, MaterialSettings.RoughnessTextureSize.Y, NewSamples, false);
+			InFlattenMaterial.RoughnessSamples.Reset(NewSamples.Num());
+			InFlattenMaterial.RoughnessSamples.Append(NewSamples);
+			InFlattenMaterial.RoughnessSize = MaterialSettings.RoughnessTextureSize;
+		}
+
+		if (InFlattenMaterial.NormalSamples.Num() && MaterialSettings.bNormalMap && InFlattenMaterial.NormalSize != MaterialSettings.NormalTextureSize)
+		{
+			TArray<FColor> NewSamples;
+			FImageUtils::ImageResize(InFlattenMaterial.NormalSize.X, InFlattenMaterial.NormalSize.Y, InFlattenMaterial.NormalSamples, MaterialSettings.NormalTextureSize.X, MaterialSettings.NormalTextureSize.Y, NewSamples, true);
+			InFlattenMaterial.NormalSamples.Reset(NewSamples.Num());
+			InFlattenMaterial.NormalSamples.Append(NewSamples);
+			InFlattenMaterial.NormalSize = MaterialSettings.NormalTextureSize;
+		}
+
+		if (InFlattenMaterial.EmissiveSamples.Num() && MaterialSettings.bEmissiveMap && InFlattenMaterial.EmissiveSize != MaterialSettings.EmissiveTextureSize)
+		{
+			TArray<FColor> NewSamples;
+			FImageUtils::ImageResize(InFlattenMaterial.EmissiveSize.X, InFlattenMaterial.EmissiveSize.Y, InFlattenMaterial.EmissiveSamples, MaterialSettings.EmissiveTextureSize.X, MaterialSettings.EmissiveTextureSize.Y, NewSamples, false);
+			InFlattenMaterial.EmissiveSamples.Reset(NewSamples.Num());
+			InFlattenMaterial.EmissiveSamples.Append(NewSamples);
+			InFlattenMaterial.EmissiveSize = MaterialSettings.EmissiveTextureSize;
+		}
+
+		if (InFlattenMaterial.OpacitySamples.Num() && MaterialSettings.bOpacityMap && InFlattenMaterial.OpacitySize != MaterialSettings.OpacityTextureSize)
+		{
+			TArray<FColor> NewSamples;
+			FImageUtils::ImageResize(InFlattenMaterial.OpacitySize.X, InFlattenMaterial.OpacitySize.Y, InFlattenMaterial.OpacitySamples, MaterialSettings.OpacityTextureSize.X, MaterialSettings.OpacityTextureSize.Y, NewSamples, false);
+			InFlattenMaterial.OpacitySamples.Reset(NewSamples.Num());
+			InFlattenMaterial.OpacitySamples.Append(NewSamples);
+			InFlattenMaterial.OpacitySize = MaterialSettings.OpacityTextureSize;
+		}
+	}
+}
+
 bool FMaterialUtilities::ExportMaterial(struct FMaterialMergeData& InMaterialData, FFlattenMaterial& OutFlattenMaterial, struct FExportMaterialProxyCache* ProxyCache)
 {
 	UMaterialInterface* Material = InMaterialData.Material;
@@ -1496,37 +1723,58 @@ bool FMaterialUtilities::ExportMaterial(struct FMaterialMergeData& InMaterialDat
 	}
 
 	// Determine whether or not certain properties can be rendered
-	bool bRenderNormal = Material->GetMaterial()->HasNormalConnected() || (Material->GetMaterial()->bUseMaterialAttributes) || Material->IsPropertyActive(MP_Normal);
-	bool bRenderEmissive = Material->GetMaterial()->EmissiveColor.IsConnected() && Material->IsPropertyActive(MP_EmissiveColor);
-	bool bRenderOpacityMask = Material->IsPropertyActive(MP_OpacityMask) && Material->GetBlendMode() == BLEND_Masked;
-	bool bRenderOpacity = Material->IsPropertyActive(MP_Opacity) && IsTranslucentBlendMode(Material->GetBlendMode());
+	bool bRenderNormal = (Material->GetMaterial()->HasNormalConnected() || Material->GetMaterial()->bUseMaterialAttributes || Material->IsPropertyActive(MP_Normal)) && (OutFlattenMaterial.NormalSize.X > 0 && OutFlattenMaterial.NormalSize.Y > 0);
+	bool bRenderEmissive = Material->GetMaterial()->EmissiveColor.IsConnected() && Material->IsPropertyActive(MP_EmissiveColor) && OutFlattenMaterial.EmissiveSize.X > 0 && OutFlattenMaterial.EmissiveSize.Y > 0;
+	bool bRenderOpacityMask = Material->IsPropertyActive(MP_OpacityMask) && Material->GetBlendMode() == BLEND_Masked && OutFlattenMaterial.OpacitySize.X > 0 && OutFlattenMaterial.OpacitySize.Y > 0;
+	bool bRenderOpacity = Material->IsPropertyActive(MP_Opacity) && IsTranslucentBlendMode(Material->GetBlendMode()) && OutFlattenMaterial.OpacitySize.X > 0 && OutFlattenMaterial.OpacitySize.Y > 0;
+	bool bRenderSubSurface = Material->IsPropertyActive(MP_SubsurfaceColor) && OutFlattenMaterial.SubSurfaceSize.X > 0 && OutFlattenMaterial.SubSurfaceSize.Y > 0;
+	bool bRenderMetallic = Material->IsPropertyActive(MP_Metallic) && OutFlattenMaterial.MetallicSize.X > 0 && OutFlattenMaterial.MetallicSize.Y > 0;
+	bool bRenderSpecular = Material->IsPropertyActive(MP_Specular) && OutFlattenMaterial.SpecularSize.X > 0 && OutFlattenMaterial.SpecularSize.Y > 0;
+	bool bRenderRoughness = Material->IsPropertyActive(MP_Roughness) && OutFlattenMaterial.RoughnessSize.X > 0 && OutFlattenMaterial.RoughnessSize.Y > 0;
+
 	check(!bRenderOpacity || !bRenderOpacityMask);
 
-	// Compile shaders and render flatten material.	
-	RenderMaterialPropertyToTexture(InMaterialData, MP_BaseColor, false, PF_B8G8R8A8, OutFlattenMaterial.DiffuseSize, OutFlattenMaterial.DiffuseSamples);
-	RenderMaterialPropertyToTexture(InMaterialData, MP_Metallic, false, PF_B8G8R8A8, OutFlattenMaterial.MetallicSize, OutFlattenMaterial.MetallicSamples);
-	RenderMaterialPropertyToTexture(InMaterialData, MP_Specular, false, PF_B8G8R8A8, OutFlattenMaterial.SpecularSize, OutFlattenMaterial.SpecularSamples);
-	RenderMaterialPropertyToTexture(InMaterialData, MP_Roughness, false, PF_B8G8R8A8, OutFlattenMaterial.RoughnessSize, OutFlattenMaterial.RoughnessSamples);
+	// Compile shaders and render flatten material.
+	RenderMaterialPropertyToTexture(InMaterialData, MP_BaseColor, false, PF_B8G8R8A8, OutFlattenMaterial.RenderSize, OutFlattenMaterial.DiffuseSize, OutFlattenMaterial.DiffuseSamples);
+
+	if (bRenderMetallic)
+	{
+		RenderMaterialPropertyToTexture(InMaterialData, MP_Metallic, false, PF_B8G8R8A8, OutFlattenMaterial.RenderSize, OutFlattenMaterial.MetallicSize, OutFlattenMaterial.MetallicSamples);
+	}	
+	if (bRenderSpecular)
+	{
+		RenderMaterialPropertyToTexture(InMaterialData, MP_Specular, false, PF_B8G8R8A8, OutFlattenMaterial.RenderSize, OutFlattenMaterial.SpecularSize, OutFlattenMaterial.SpecularSamples);
+	}
+	if (bRenderRoughness)
+	{
+		RenderMaterialPropertyToTexture(InMaterialData, MP_Roughness, false, PF_B8G8R8A8, OutFlattenMaterial.RenderSize, OutFlattenMaterial.RoughnessSize, OutFlattenMaterial.RoughnessSamples);
+	}
 	if (bRenderNormal)
 	{
-		RenderMaterialPropertyToTexture(InMaterialData, MP_Normal, true, PF_B8G8R8A8, OutFlattenMaterial.NormalSize, OutFlattenMaterial.NormalSamples);
+		RenderMaterialPropertyToTexture(InMaterialData, MP_Normal, true, PF_B8G8R8A8, OutFlattenMaterial.RenderSize, OutFlattenMaterial.NormalSize, OutFlattenMaterial.NormalSamples);
 	}
 	if (bRenderOpacityMask)
 	{
-		RenderMaterialPropertyToTexture(InMaterialData, MP_OpacityMask, true, PF_B8G8R8A8, OutFlattenMaterial.OpacitySize, OutFlattenMaterial.OpacitySamples);
+		RenderMaterialPropertyToTexture(InMaterialData, MP_OpacityMask, true, PF_B8G8R8A8, OutFlattenMaterial.RenderSize, OutFlattenMaterial.OpacitySize, OutFlattenMaterial.OpacitySamples);
 	}
 	if (bRenderOpacity)
 	{
 		// Number of blend modes, let's UMaterial decide whether it wants this property
-		RenderMaterialPropertyToTexture(InMaterialData, MP_Opacity, true, PF_B8G8R8A8, OutFlattenMaterial.OpacitySize, OutFlattenMaterial.OpacitySamples);
+		RenderMaterialPropertyToTexture(InMaterialData, MP_Opacity, true, PF_B8G8R8A8, OutFlattenMaterial.RenderSize, OutFlattenMaterial.OpacitySize, OutFlattenMaterial.OpacitySamples);
 	}
 	if (bRenderEmissive)
 	{
 		// PF_FloatRGBA is here to be able to render and read HDR image using ReadFloat16Pixels()
-		RenderMaterialPropertyToTexture(InMaterialData, MP_EmissiveColor, false, PF_FloatRGBA, OutFlattenMaterial.EmissiveSize, OutFlattenMaterial.EmissiveSamples);
+		RenderMaterialPropertyToTexture(InMaterialData, MP_EmissiveColor, false, PF_FloatRGBA, OutFlattenMaterial.RenderSize, OutFlattenMaterial.EmissiveSize, OutFlattenMaterial.EmissiveSamples);
 		OutFlattenMaterial.EmissiveScale = InMaterialData.EmissiveScale;
 	}	
 
+	if (bRenderSubSurface)
+	{
+		// TODO support rendering out sub surface color property
+		/*RenderMaterialPropertyToTexture(InMaterialData, MP_SubsurfaceColor, false, PF_B8G8R8A8, OutFlattenMaterial.RenderSize, OutFlattenMaterial.SubSurfaceSize, OutFlattenMaterial.SubSurfaceSamples);*/
+	}
+
 	OutFlattenMaterial.MaterialId = Material->GetLightingGuid();
 
 	// Swap back the proxy cache
@@ -1541,13 +1789,15 @@ bool FMaterialUtilities::ExportMaterial(struct FMaterialMergeData& InMaterialDat
 	return true;
 }
 
-bool FMaterialUtilities::RenderMaterialPropertyToTexture(struct FMaterialMergeData& InMaterialData, EMaterialProperty InMaterialProperty, bool bInForceLinearGamma, EPixelFormat InPixelFormat, FIntPoint& InTargetSize, TArray<FColor>& OutSamples)
+bool FMaterialUtilities::RenderMaterialPropertyToTexture(struct FMaterialMergeData& InMaterialData, EMaterialProperty InMaterialProperty, bool bInForceLinearGamma, EPixelFormat InPixelFormat, const FIntPoint& InTargetSize, FIntPoint& OutSampleSize, TArray<FColor>& OutSamples)
 {
 	if (InTargetSize.X == 0 || InTargetSize.Y == 0)
 	{
 		return false;
 	}
 
+	OutSampleSize = InTargetSize;
+
 	FMaterialRenderProxy* MaterialProxy = nullptr;
 
 	check(InMaterialProperty >= 0 && InMaterialProperty < ARRAY_COUNT(InMaterialData.ProxyCache->Proxies));
@@ -1569,7 +1819,7 @@ bool FMaterialUtilities::RenderMaterialPropertyToTexture(struct FMaterialMergeDa
 	check(CurrentlyRendering == false);
 	CurrentlyRendering = true;
 
-	UTextureRenderTarget2D* RenderTarget = CreateRenderTarget(bInForceLinearGamma, InPixelFormat, InTargetSize);
+	UTextureRenderTarget2D* RenderTarget = CreateRenderTarget(bInForceLinearGamma, InPixelFormat, OutSampleSize);
 
 	OutSamples.Empty(InTargetSize.X * InTargetSize.Y);
 	bool bResult = FMeshRenderer::RenderMaterial(
@@ -1620,7 +1870,7 @@ bool FMaterialUtilities::RenderMaterialPropertyToTexture(struct FMaterialMergeDa
 	// Uniform value
 	if (bIsUniform)
 	{
-		InTargetSize = FIntPoint(1, 1);
+		OutSampleSize = FIntPoint(1, 1);
 		OutSamples.Empty();
 		OutSamples.Add(MaxColor);
 	}
diff --git a/Engine/Source/Developer/MaterialUtilities/Public/MaterialUtilities.h b/Engine/Source/Developer/MaterialUtilities/Public/MaterialUtilities.h
--- a/Engine/Source/Developer/MaterialUtilities/Public/MaterialUtilities.h
+++ b/Engine/Source/Developer/MaterialUtilities/Public/MaterialUtilities.h
@@ -14,13 +14,15 @@
 struct FFlattenMaterial
 {
 	FFlattenMaterial()
-		: DiffuseSize(0, 0)
+		: RenderSize( 0, 0 )
+		, DiffuseSize(0, 0)
 		, NormalSize(0, 0)
 		, MetallicSize(0, 0)
 		, RoughnessSize(0, 0)
 		, SpecularSize(0, 0)
 		, OpacitySize(0, 0)
 		, EmissiveSize(0, 0)
+		, SubSurfaceSize(0, 0)
 		, bTwoSided(false)
 		, BlendMode(BLEND_Opaque)
 		, EmissiveScale(1.0f)		
@@ -36,6 +38,7 @@ struct FFlattenMaterial
 		SpecularSamples.Empty();
 		OpacitySamples.Empty();
 		EmissiveSamples.Empty();
+		SubSurfaceSamples.Empty();
 	}
 	
 	/** Set all alpha channel values with InAlphaValue */
@@ -69,12 +72,17 @@ struct FFlattenMaterial
 		{ 
 			Sample.A = InAlphaValue; 
 		}
+		for (FColor& Sample : SubSurfaceSamples)
+		{
+			Sample.A = InAlphaValue;
+		}
 	}
 	
 	/** Material Guid */
 	FGuid			MaterialId;
 		
 	/** Texture sizes for each individual property*/
+	FIntPoint		RenderSize;
 	FIntPoint		DiffuseSize;
 	FIntPoint		NormalSize;
 	FIntPoint		MetallicSize;	
@@ -82,6 +90,7 @@ struct FFlattenMaterial
 	FIntPoint		SpecularSize;	
 	FIntPoint		OpacitySize;
 	FIntPoint		EmissiveSize;
+	FIntPoint		SubSurfaceSize;
 
 	/** Baked down texture samples for each individual property*/
 	TArray<FColor>	DiffuseSamples;
@@ -91,6 +100,7 @@ struct FFlattenMaterial
 	TArray<FColor>	SpecularSamples;
 	TArray<FColor>	OpacitySamples;
 	TArray<FColor>	EmissiveSamples;
+	TArray<FColor>	SubSurfaceSamples;
 
 	/** Flag whether or not the material will have to be two-sided */
 	bool			bTwoSided;
@@ -105,7 +115,7 @@ struct MATERIALUTILITIES_API FExportMaterialProxyCache
 {
 	// Material proxies for each property. Note: we're not handling all properties here,
 	// so hold only up to MP_Normal inclusive.
-	class FMaterialRenderProxy* Proxies[EMaterialProperty::MP_Normal + 1];
+	class FMaterialRenderProxy* Proxies[EMaterialProperty::MP_MAX];
 
 	FExportMaterialProxyCache();
 	~FExportMaterialProxyCache();
@@ -280,9 +290,9 @@ public:
 	* @param InMaterial 			Material to analyze
 	* @param InMaterialSettings 	Settings containing how to material should be merged
 	* @param OutNumTexCoords		Number of texture coordinates used across all properties flagged for merging
-	* @param OutUseVertexColor		Flag whether or not Vertex Colors are used in the material graph for the properties flagged for merging
+	* @param OutRequiresVertexData	Flag whether or not Vertex Data is used in the material graph for the properties flagged for merging
 	*/
-	static void AnalyzeMaterial(class UMaterialInterface* InMaterial, const struct FMaterialProxySettings& InMaterialSettings, int32& OutNumTexCoords, bool& OutUseVertexColor);
+	static void AnalyzeMaterial(class UMaterialInterface* InMaterial, const struct FMaterialProxySettings& InMaterialSettings, int32& OutNumTexCoords, bool& OutRequiresVertexData);
 
 	/**
 	* Remaps material indices where possible to reduce the number of materials required for creating a proxy material
@@ -305,6 +315,13 @@ public:
 	* @param InFlattenMaterial				Flatten material to optimize
 	*/
 	static void OptimizeFlattenMaterial(FFlattenMaterial& InFlattenMaterial);
+
+	/**
+	* Resizes flatten material's data if applicable by comparing it with the original settings
+	*
+	* @param InFlattenMaterial				Flatten material to optimize
+	*/
+	static void ResizeFlattenMaterial(FFlattenMaterial& InFlattenMaterial, const struct FMeshProxySettings& InMeshProxySettings);
 private:
 	
 	/**
@@ -324,10 +341,11 @@ private:
 	* @param bInForceLinearGamma	Whether or not to force linear gamma (used for Normal property)
 	* @param InPixelFormat			Pixel format of the target texture
 	* @param InTargetSize			Dimensions of the target texture
+	* @param OutSampleSize			Dimensions of the rendered texture
 	* @param OutSamples				Array of FColor samples containing the rendered out texture pixel data
 	* @return						Whether operation was successful
 	*/
-	static bool RenderMaterialPropertyToTexture(struct FMaterialMergeData& InMaterialData, EMaterialProperty InMaterialProperty, bool bInForceLinearGamma, EPixelFormat InPixelFormat, FIntPoint& InTargetSize, TArray<FColor>& OutSamples);
+	static bool RenderMaterialPropertyToTexture(struct FMaterialMergeData& InMaterialData, EMaterialProperty InMaterialProperty, bool bInForceLinearGamma, EPixelFormat InPixelFormat, const FIntPoint& InTargetSize, FIntPoint& OutSampleSize, TArray<FColor>& OutSamples);
 
 	/**
 	* Creates and add or reuses a RenderTarget from the pool
