Copying //UE4/Dev-Framework to //UE4/Dev-Main (Source: //UE4/Dev-Framework @ 3136612)

diff --git a/Engine/Source/Developer/MaterialUtilities/Private/MaterialUtilities.cpp b/Engine/Source/Developer/MaterialUtilities/Private/MaterialUtilities.cpp
--- a/Engine/Source/Developer/MaterialUtilities/Private/MaterialUtilities.cpp
+++ b/Engine/Source/Developer/MaterialUtilities/Private/MaterialUtilities.cpp
@@ -778,55 +778,60 @@ bool FMaterialUtilities::ExportLandscapeMaterial(ALandscapeProxy* InLandscape, c
 
 	FSceneInterface* Scene = InLandscape->GetWorld()->Scene;
 						
-	// Render diffuse texture using BufferVisualizationMode=BaseColor
-	if (OutFlattenMaterial.DiffuseSize.X > 0 && 
-		OutFlattenMaterial.DiffuseSize.Y > 0)
+	// Render diffuse texture using BufferVisualizationMode=BaseColor	
+	if (OutFlattenMaterial.ShouldGenerateDataForProperty(EFlattenMaterialProperties::Diffuse))
 	{
+		const FIntPoint& DiffuseSize = OutFlattenMaterial.GetPropertySize(EFlattenMaterialProperties::Diffuse);
 		static const FName BaseColorName("BaseColor");
 		const float BaseColorGamma = 2.2f; // BaseColor to gamma space
+		TArray<FColor>& DiffuseSamples = OutFlattenMaterial.GetPropertySamples(EFlattenMaterialProperties::Diffuse);
 		RenderSceneToTexture(Scene, BaseColorName, ViewOrigin, ViewRotationMatrix, ProjectionMatrix, HiddenPrimitives, 
-			OutFlattenMaterial.DiffuseSize, BaseColorGamma, OutFlattenMaterial.DiffuseSamples);
+			DiffuseSize, BaseColorGamma, DiffuseSamples);
 	}
 
 	// Render normal map using BufferVisualizationMode=WorldNormal
 	// Final material should use world space instead of tangent space for normals
-	if (OutFlattenMaterial.NormalSize.X > 0 && 
-		OutFlattenMaterial.NormalSize.Y > 0)
+	if (OutFlattenMaterial.ShouldGenerateDataForProperty(EFlattenMaterialProperties::Normal))
 	{
 		static const FName WorldNormalName("WorldNormal");
 		const float NormalColorGamma = 1.0f; // Dump normal texture in linear space
+		const FIntPoint& NormalSize = OutFlattenMaterial.GetPropertySize(EFlattenMaterialProperties::Normal);
+		TArray<FColor>& NormalSamples = OutFlattenMaterial.GetPropertySamples(EFlattenMaterialProperties::Normal);
 		RenderSceneToTexture(Scene, WorldNormalName, ViewOrigin, ViewRotationMatrix, ProjectionMatrix, HiddenPrimitives, 
-			OutFlattenMaterial.NormalSize, NormalColorGamma, OutFlattenMaterial.NormalSamples);
+			NormalSize, NormalColorGamma, NormalSamples);
 	}
 
-	// Render metallic map using BufferVisualizationMode=Metallic
-	if (OutFlattenMaterial.MetallicSize.X > 0 && 
-		OutFlattenMaterial.MetallicSize.Y > 0)
+	// Render metallic map using BufferVisualizationMode=Metallic	
+	if (OutFlattenMaterial.ShouldGenerateDataForProperty(EFlattenMaterialProperties::Metallic))
 	{
 		static const FName MetallicName("Metallic");
 		const float MetallicColorGamma = 1.0f; // Dump metallic texture in linear space
+		const FIntPoint& MetallicSize = OutFlattenMaterial.GetPropertySize(EFlattenMaterialProperties::Metallic);
+		TArray<FColor>& MetallicSamples = OutFlattenMaterial.GetPropertySamples(EFlattenMaterialProperties::Metallic);
 		RenderSceneToTexture(Scene, MetallicName, ViewOrigin, ViewRotationMatrix, ProjectionMatrix, HiddenPrimitives, 
-			OutFlattenMaterial.MetallicSize, MetallicColorGamma, OutFlattenMaterial.MetallicSamples);
+			MetallicSize, MetallicColorGamma, MetallicSamples);
 	}
 
-	// Render roughness map using BufferVisualizationMode=Roughness
-	if (OutFlattenMaterial.RoughnessSize.X > 0 && 
-		OutFlattenMaterial.RoughnessSize.Y > 0)
+	// Render roughness map using BufferVisualizationMode=Roughness	
+	if (OutFlattenMaterial.ShouldGenerateDataForProperty(EFlattenMaterialProperties::Roughness))
 	{
 		static const FName RoughnessName("Roughness");
 		const float RoughnessColorGamma = 2.2f; // Roughness material powers color by 2.2, transform it back to linear
+		const FIntPoint& RoughnessSize = OutFlattenMaterial.GetPropertySize(EFlattenMaterialProperties::Roughness);
+		TArray<FColor>& RoughnessSamples = OutFlattenMaterial.GetPropertySamples(EFlattenMaterialProperties::Roughness);
 		RenderSceneToTexture(Scene, RoughnessName, ViewOrigin, ViewRotationMatrix, ProjectionMatrix, HiddenPrimitives, 
-			OutFlattenMaterial.RoughnessSize, RoughnessColorGamma, OutFlattenMaterial.RoughnessSamples);
+			RoughnessSize, RoughnessColorGamma, RoughnessSamples);
 	}
 
-	// Render specular map using BufferVisualizationMode=Specular
-	if (OutFlattenMaterial.SpecularSize.X > 0 && 
-		OutFlattenMaterial.SpecularSize.Y > 0)
+	// Render specular map using BufferVisualizationMode=Specular	
+	if (OutFlattenMaterial.ShouldGenerateDataForProperty(EFlattenMaterialProperties::Specular))
 	{
 		static const FName SpecularName("Specular");
 		const float SpecularColorGamma = 1.0f; // Dump specular texture in linear space
+		const FIntPoint& SpecularSize = OutFlattenMaterial.GetPropertySize(EFlattenMaterialProperties::Specular);
+		TArray<FColor>& SpecularSamples = OutFlattenMaterial.GetPropertySamples(EFlattenMaterialProperties::Specular);
 		RenderSceneToTexture(Scene, SpecularName, ViewOrigin, ViewRotationMatrix, ProjectionMatrix, HiddenPrimitives, 
-			OutFlattenMaterial.SpecularSize, SpecularColorGamma, OutFlattenMaterial.SpecularSamples);
+			SpecularSize, SpecularColorGamma, SpecularSamples);
 	}
 				
 	OutFlattenMaterial.MaterialId = InLandscape->GetLandscapeGuid();
@@ -866,12 +871,13 @@ UMaterial* FMaterialUtilities::CreateMaterial(const FFlattenMaterial& InFlattenM
 	int32 MaterialNodeStepY = 180;
 
 	// BaseColor
-	if (InFlattenMaterial.DiffuseSamples.Num() > 1)
+	const TArray<FColor>& DiffuseSamples = InFlattenMaterial.GetPropertySamples(EFlattenMaterialProperties::Diffuse);
+	if (DiffuseSamples.Num() > 1)
 	{
 		const FString AssetName = TEXT("T_") + AssetBaseName + TEXT("_D");
 		const FString AssetLongName = AssetBasePath + AssetName;
 		const bool bSRGB = true;
-		UTexture2D* Texture = CreateTexture(InOuter, AssetLongName, InFlattenMaterial.DiffuseSize, InFlattenMaterial.DiffuseSamples, TC_Default, InTextureGroup, Flags, bSRGB);
+		UTexture2D* Texture = CreateTexture(InOuter, AssetLongName, InFlattenMaterial.GetPropertySize(EFlattenMaterialProperties::Diffuse), DiffuseSamples, TC_Default, InTextureGroup, Flags, bSRGB);
 		OutGeneratedAssets.Add(Texture);
 			
 		auto BasecolorExpression = NewObject<UMaterialExpressionTextureSample>(Material);
@@ -884,10 +890,10 @@ UMaterial* FMaterialUtilities::CreateMaterial(const FFlattenMaterial& InFlattenM
 
 		MaterialNodeY+= MaterialNodeStepY;
 	}
-	else if (InFlattenMaterial.DiffuseSamples.Num() == 1)
+	else if (DiffuseSamples.Num() == 1)
 	{
 		// Set Roughness to constant
-		FLinearColor BaseColor = FLinearColor(InFlattenMaterial.DiffuseSamples[0]);
+		FLinearColor BaseColor = FLinearColor(DiffuseSamples[0]);
 		auto BaseColorExpression = NewObject<UMaterialExpressionConstant4Vector>(Material);
 		BaseColorExpression->Constant = BaseColor;
 		BaseColorExpression->MaterialExpressionEditorX = -400;
@@ -900,9 +906,9 @@ UMaterial* FMaterialUtilities::CreateMaterial(const FFlattenMaterial& InFlattenM
 
 
 	// Whether or not a material property is baked down
-	const bool bHasMetallic = (InFlattenMaterial.MetallicSamples.Num() > 1);
-	const bool bHasRoughness = (InFlattenMaterial.RoughnessSamples.Num() > 1);
-	const bool bHasSpecular = (InFlattenMaterial.SpecularSamples.Num() > 1);
+	const bool bHasMetallic = InFlattenMaterial.DoesPropertyContainData(EFlattenMaterialProperties::Metallic) && !InFlattenMaterial.IsPropertyConstant(EFlattenMaterialProperties::Metallic);
+	const bool bHasRoughness = InFlattenMaterial.DoesPropertyContainData(EFlattenMaterialProperties::Roughness) && !InFlattenMaterial.IsPropertyConstant(EFlattenMaterialProperties::Roughness);
+	const bool bHasSpecular = InFlattenMaterial.DoesPropertyContainData(EFlattenMaterialProperties::Specular) && !InFlattenMaterial.IsPropertyConstant(EFlattenMaterialProperties::Specular);
 
 	// Number of material properties baked down to textures
 	const int BakedMaterialPropertyCount = bHasMetallic + bHasRoughness + bHasSpecular;
@@ -912,18 +918,17 @@ UMaterial* FMaterialUtilities::CreateMaterial(const FFlattenMaterial& InFlattenM
 
 	int32 SampleCount = 0;
 	FIntPoint MergedSize(0,0);
-	SampleCount = (bHasMetallic && SampleCount == 0) ? InFlattenMaterial.MetallicSamples.Num() : SampleCount;
-	MergedSize = (bHasMetallic && MergedSize.X == 0) ? InFlattenMaterial.MetallicSize : MergedSize;
-
-	SampleCount = (bHasRoughness && SampleCount == 0) ? InFlattenMaterial.RoughnessSamples.Num() : SampleCount;
-	MergedSize = (bHasRoughness && MergedSize.X == 0) ? InFlattenMaterial.RoughnessSize : MergedSize;
-
-	SampleCount = (bHasSpecular && SampleCount == 0) ? InFlattenMaterial.SpecularSamples.Num() : SampleCount;
-	MergedSize = (bHasSpecular && MergedSize.X == 0) ? InFlattenMaterial.SpecularSize : MergedSize;
-
-	bSameTextureSize &= bHasMetallic ? (SampleCount == InFlattenMaterial.MetallicSamples.Num()) : true;
-	bSameTextureSize &= bHasRoughness ? (SampleCount == InFlattenMaterial.RoughnessSamples.Num()) : true;
-	bSameTextureSize &= bHasSpecular ? (SampleCount == InFlattenMaterial.SpecularSamples.Num()) : true;
+	for (int32 PropertyIndex = 0; PropertyIndex < 3; ++PropertyIndex)
+	{
+		EFlattenMaterialProperties Property = (EFlattenMaterialProperties)(PropertyIndex + (int32)EFlattenMaterialProperties::Metallic);
+		const bool HasProperty = InFlattenMaterial.DoesPropertyContainData(Property) && !InFlattenMaterial.IsPropertyConstant(Property);
+		FIntPoint PropertySize = InFlattenMaterial.GetPropertySize(Property);
+		SampleCount = (bHasMetallic && SampleCount == 0) ? InFlattenMaterial.GetPropertySamples(Property).Num() : SampleCount;
+		MergedSize = (bHasMetallic && MergedSize.X == 0) ? PropertySize : MergedSize;
+	}
+	bSameTextureSize &= bHasMetallic ? (SampleCount == InFlattenMaterial.GetPropertySamples(EFlattenMaterialProperties::Metallic).Num()) : true;
+	bSameTextureSize &= bHasRoughness ? (SampleCount == InFlattenMaterial.GetPropertySamples(EFlattenMaterialProperties::Roughness).Num()) : true;
+	bSameTextureSize &= bHasSpecular ? (SampleCount == InFlattenMaterial.GetPropertySamples(EFlattenMaterialProperties::Specular).Num()) : true;
 
 	// Merge values into one texture if more than one material property exists
 	if (BakedMaterialPropertyCount > 1 && bSameTextureSize)
@@ -931,20 +936,27 @@ UMaterial* FMaterialUtilities::CreateMaterial(const FFlattenMaterial& InFlattenM
 		// Metallic = R, Roughness = G, Specular = B
 		TArray<FColor> MergedSamples;
 		MergedSamples.AddZeroed(SampleCount);
-		
-		if (bHasMetallic) for (int32 SampleIndex = 0; SampleIndex < SampleCount; ++SampleIndex)
-		{
-			MergedSamples[SampleIndex].R = InFlattenMaterial.MetallicSamples[SampleIndex].R;
-		}
-
-		if (bHasRoughness) for (int32 SampleIndex = 0; SampleIndex < SampleCount; ++SampleIndex)
-		{
-			MergedSamples[SampleIndex].G = InFlattenMaterial.RoughnessSamples[SampleIndex].G;
-		}
 
-		if (bHasSpecular) for (int32 SampleIndex = 0; SampleIndex < SampleCount; ++SampleIndex)
+		// R G B masks
+#if PLATFORM_LITTLE_ENDIAN
+		const uint32 ColorMask[3] = { 0x0000FF00, 0x00FF0000, 0xFF000000 };
+#else // PLATFORM_LITTLE_ENDIAN
+		const uint32 ColorMask[3] = { 0x00FF0000, 0x0000FF00, 0x000000FF };
+#endif
+		for (int32 PropertyIndex = 0; PropertyIndex < 3; ++PropertyIndex)
 		{
-			MergedSamples[SampleIndex].B = InFlattenMaterial.SpecularSamples[SampleIndex].B;
+			EFlattenMaterialProperties Property = (EFlattenMaterialProperties)(PropertyIndex + (int32)EFlattenMaterialProperties::Metallic);
+			const bool HasProperty = InFlattenMaterial.DoesPropertyContainData(Property) && !InFlattenMaterial.IsPropertyConstant(Property);
+			
+			if (HasProperty)
+			{
+				const TArray<FColor>& PropertySamples = InFlattenMaterial.GetPropertySamples(Property);
+				// OR masked values (samples initialized to zero, so no random data)
+				for (int32 SampleIndex = 0; SampleIndex < SampleCount; ++SampleIndex)
+				{
+					MergedSamples[SampleIndex].DWColor() |= (PropertySamples[SampleIndex].DWColor() & ColorMask[PropertyIndex]);
+				}
+			}
 		}
 
 		const FString AssetName = TEXT("T_") + AssetBaseName + TEXT("_MRS");
@@ -997,11 +1009,11 @@ UMaterial* FMaterialUtilities::CreateMaterial(const FFlattenMaterial& InFlattenM
 	else
 	{
 		// Metallic
-		if (InFlattenMaterial.MetallicSamples.Num() > 1 && MaterialProxySettings.bMetallicMap)
+		if (bHasMetallic && MaterialProxySettings.bMetallicMap)
 		{
 			const FString AssetName = TEXT("T_") + AssetBaseName + TEXT("_M");
 			const bool bSRGB = true;
-			UTexture2D* Texture = CreateTexture(InOuter, AssetBasePath + AssetName, InFlattenMaterial.MetallicSize, InFlattenMaterial.MetallicSamples, TC_Default, InTextureGroup, Flags, bSRGB);
+			UTexture2D* Texture = CreateTexture(InOuter, AssetBasePath + AssetName, InFlattenMaterial.GetPropertySize(EFlattenMaterialProperties::Metallic), InFlattenMaterial.GetPropertySamples(EFlattenMaterialProperties::Metallic), TC_Default, InTextureGroup, Flags, bSRGB);
 			OutGeneratedAssets.Add(Texture);
 
 			auto MetallicExpression = NewObject<UMaterialExpressionTextureSample>(Material);
@@ -1016,11 +1028,11 @@ UMaterial* FMaterialUtilities::CreateMaterial(const FFlattenMaterial& InFlattenM
 		}
 
 		// Specular
-		if (InFlattenMaterial.SpecularSamples.Num() > 1 && MaterialProxySettings.bSpecularMap)
+		if (bHasSpecular && MaterialProxySettings.bSpecularMap)
 		{
 			const FString AssetName = TEXT("T_") + AssetBaseName + TEXT("_S");
 			const bool bSRGB = true;
-			UTexture2D* Texture = CreateTexture(InOuter, AssetBasePath + AssetName, InFlattenMaterial.SpecularSize, InFlattenMaterial.SpecularSamples, TC_Default, InTextureGroup, Flags, bSRGB);
+			UTexture2D* Texture = CreateTexture(InOuter, AssetBasePath + AssetName, InFlattenMaterial.GetPropertySize(EFlattenMaterialProperties::Specular), InFlattenMaterial.GetPropertySamples(EFlattenMaterialProperties::Specular), TC_Default, InTextureGroup, Flags, bSRGB);
 			OutGeneratedAssets.Add(Texture);
 
 			auto SpecularExpression = NewObject<UMaterialExpressionTextureSample>(Material);
@@ -1035,11 +1047,11 @@ UMaterial* FMaterialUtilities::CreateMaterial(const FFlattenMaterial& InFlattenM
 		}
 
 		// Roughness
-		if (InFlattenMaterial.RoughnessSamples.Num() > 1 && MaterialProxySettings.bRoughnessMap)
+		if (bHasRoughness && MaterialProxySettings.bRoughnessMap)
 		{
 			const FString AssetName = TEXT("T_") + AssetBaseName + TEXT("_R");
 			const bool bSRGB = true;
-			UTexture2D* Texture = CreateTexture(InOuter, AssetBasePath + AssetName, InFlattenMaterial.RoughnessSize, InFlattenMaterial.RoughnessSamples, TC_Default, InTextureGroup, Flags, bSRGB);
+			UTexture2D* Texture = CreateTexture(InOuter, AssetBasePath + AssetName, InFlattenMaterial.GetPropertySize(EFlattenMaterialProperties::Roughness), InFlattenMaterial.GetPropertySamples(EFlattenMaterialProperties::Roughness), TC_Default, InTextureGroup, Flags, bSRGB);
 			OutGeneratedAssets.Add(Texture);
 
 			auto RoughnessExpression = NewObject<UMaterialExpressionTextureSample>(Material);
@@ -1054,10 +1066,10 @@ UMaterial* FMaterialUtilities::CreateMaterial(const FFlattenMaterial& InFlattenM
 		}
 	}
 
-	if (InFlattenMaterial.MetallicSamples.Num() == 1 || !MaterialProxySettings.bMetallicMap)
+	if (InFlattenMaterial.IsPropertyConstant(EFlattenMaterialProperties::Metallic) || !MaterialProxySettings.bMetallicMap)
 	{
 		auto MetallicExpression = NewObject<UMaterialExpressionConstant>(Material);
-		MetallicExpression->R = MaterialProxySettings.bMetallicMap ? FLinearColor(InFlattenMaterial.MetallicSamples[0]).R : MaterialProxySettings.MetallicConstant;
+		MetallicExpression->R = MaterialProxySettings.bMetallicMap ? FLinearColor(InFlattenMaterial.GetPropertySamples(EFlattenMaterialProperties::Metallic)[0]).R : MaterialProxySettings.MetallicConstant;
 		MetallicExpression->MaterialExpressionEditorX = -400;
 		MetallicExpression->MaterialExpressionEditorY = MaterialNodeY;
 		Material->Expressions.Add(MetallicExpression);
@@ -1066,11 +1078,11 @@ UMaterial* FMaterialUtilities::CreateMaterial(const FFlattenMaterial& InFlattenM
 		MaterialNodeY += MaterialNodeStepY;
 	}
 
-	if (InFlattenMaterial.SpecularSamples.Num() == 1 || !MaterialProxySettings.bSpecularMap)
+	if (InFlattenMaterial.IsPropertyConstant(EFlattenMaterialProperties::Specular ) || !MaterialProxySettings.bSpecularMap)
 	{
 		// Set Specular to constant
 		auto SpecularExpression = NewObject<UMaterialExpressionConstant>(Material);
-		SpecularExpression->R = MaterialProxySettings.bSpecularMap ? FLinearColor(InFlattenMaterial.SpecularSamples[0]).R : MaterialProxySettings.SpecularConstant;
+		SpecularExpression->R = MaterialProxySettings.bSpecularMap ? FLinearColor(InFlattenMaterial.GetPropertySamples(EFlattenMaterialProperties::Specular)[0]).R : MaterialProxySettings.SpecularConstant;
 		SpecularExpression->MaterialExpressionEditorX = -400;
 		SpecularExpression->MaterialExpressionEditorY = MaterialNodeY;
 		Material->Expressions.Add(SpecularExpression);
@@ -1079,11 +1091,11 @@ UMaterial* FMaterialUtilities::CreateMaterial(const FFlattenMaterial& InFlattenM
 		MaterialNodeY += MaterialNodeStepY;
 	}
 	
-	if (InFlattenMaterial.RoughnessSamples.Num() == 1 || !MaterialProxySettings.bRoughnessMap)
+	if (InFlattenMaterial.IsPropertyConstant(EFlattenMaterialProperties::Roughness) || !MaterialProxySettings.bRoughnessMap)
 	{
 		// Set Roughness to constant
 		auto RoughnessExpression = NewObject<UMaterialExpressionConstant>(Material);
-		RoughnessExpression->R = MaterialProxySettings.bRoughnessMap ? FLinearColor(InFlattenMaterial.RoughnessSamples[0]).R : MaterialProxySettings.RoughnessConstant;
+		RoughnessExpression->R = MaterialProxySettings.bRoughnessMap ? FLinearColor(InFlattenMaterial.GetPropertySamples(EFlattenMaterialProperties::Roughness)[0]).R : MaterialProxySettings.RoughnessConstant;
 		RoughnessExpression->MaterialExpressionEditorX = -400;
 		RoughnessExpression->MaterialExpressionEditorY = MaterialNodeY;
 		Material->Expressions.Add(RoughnessExpression);
@@ -1093,11 +1105,11 @@ UMaterial* FMaterialUtilities::CreateMaterial(const FFlattenMaterial& InFlattenM
 	}
 
 	// Normal
-	if (InFlattenMaterial.NormalSamples.Num() > 1)
+	if (InFlattenMaterial.DoesPropertyContainData(EFlattenMaterialProperties::Normal) && !InFlattenMaterial.IsPropertyConstant(EFlattenMaterialProperties::Normal))
 	{
 		const FString AssetName = TEXT("T_") + AssetBaseName + TEXT("_N");
 		const bool bSRGB = false;
-		UTexture2D* Texture = CreateTexture(InOuter, AssetBasePath + AssetName, InFlattenMaterial.NormalSize, InFlattenMaterial.NormalSamples, TC_Normalmap, (InTextureGroup != TEXTUREGROUP_World) ? InTextureGroup : TEXTUREGROUP_WorldNormalMap, Flags, bSRGB);
+		UTexture2D* Texture = CreateTexture(InOuter, AssetBasePath + AssetName, InFlattenMaterial.GetPropertySize(EFlattenMaterialProperties::Normal), InFlattenMaterial.GetPropertySamples(EFlattenMaterialProperties::Normal), TC_Normalmap, (InTextureGroup != TEXTUREGROUP_World) ? InTextureGroup : TEXTUREGROUP_WorldNormalMap, Flags, bSRGB);
 		OutGeneratedAssets.Add(Texture);
 			
 		auto NormalExpression = NewObject<UMaterialExpressionTextureSample>(Material);
@@ -1111,10 +1123,10 @@ UMaterial* FMaterialUtilities::CreateMaterial(const FFlattenMaterial& InFlattenM
 		MaterialNodeY+= MaterialNodeStepY;
 	}
 
-	if (InFlattenMaterial.EmissiveSamples.Num() == 1)
+	if (InFlattenMaterial.IsPropertyConstant(EFlattenMaterialProperties::Emissive))
 	{
 		// Set Emissive to constant
-		FColor EmissiveColor = InFlattenMaterial.EmissiveSamples[0];
+		FColor EmissiveColor = InFlattenMaterial.GetPropertySamples(EFlattenMaterialProperties::Emissive)[0];
 		// Don't have to deal with black emissive color
 		if (EmissiveColor != FColor(0, 0, 0, 255))
 		{
@@ -1128,11 +1140,11 @@ UMaterial* FMaterialUtilities::CreateMaterial(const FFlattenMaterial& InFlattenM
 			MaterialNodeY += MaterialNodeStepY;
 		}
 	}
-	else if (InFlattenMaterial.EmissiveSamples.Num() > 1)
+	else if (InFlattenMaterial.DoesPropertyContainData(EFlattenMaterialProperties::Emissive) && !InFlattenMaterial.IsPropertyConstant(EFlattenMaterialProperties::Emissive))
 	{
 		const FString AssetName = TEXT("T_") + AssetBaseName + TEXT("_E");
 		const bool bSRGB = true;
-		UTexture2D* Texture = CreateTexture(InOuter, AssetBasePath + AssetName, InFlattenMaterial.EmissiveSize, InFlattenMaterial.EmissiveSamples, TC_Default, InTextureGroup, Flags, bSRGB);
+		UTexture2D* Texture = CreateTexture(InOuter, AssetBasePath + AssetName, InFlattenMaterial.GetPropertySize(EFlattenMaterialProperties::Emissive), InFlattenMaterial.GetPropertySamples(EFlattenMaterialProperties::Emissive), TC_Default, InTextureGroup, Flags, bSRGB);
 		OutGeneratedAssets.Add(Texture);
 
 		//Assign emissive color to the material
@@ -1154,10 +1166,10 @@ UMaterial* FMaterialUtilities::CreateMaterial(const FFlattenMaterial& InFlattenM
 		MaterialNodeY += MaterialNodeStepY;
 	}
 
-	if (InFlattenMaterial.OpacitySamples.Num() == 1)
+	if (InFlattenMaterial.IsPropertyConstant(EFlattenMaterialProperties::Opacity))
 	{
 		// Set Opacity to constant
-		FLinearColor Opacity = FLinearColor(InFlattenMaterial.OpacitySamples[0]);
+		FLinearColor Opacity = FLinearColor(InFlattenMaterial.GetPropertySamples(EFlattenMaterialProperties::Opacity)[0]);
 		auto OpacityExpression = NewObject<UMaterialExpressionConstant>(Material);
 		OpacityExpression->R = Opacity.R;
 		OpacityExpression->MaterialExpressionEditorX = -400;
@@ -1167,11 +1179,11 @@ UMaterial* FMaterialUtilities::CreateMaterial(const FFlattenMaterial& InFlattenM
 
 		MaterialNodeY += MaterialNodeStepY;
 	}
-	else if (InFlattenMaterial.OpacitySamples.Num() > 1)
+	else if (InFlattenMaterial.DoesPropertyContainData(EFlattenMaterialProperties::Opacity))
 	{
 		const FString AssetName = TEXT("T_") + AssetBaseName + TEXT("_O");
 		const bool bSRGB = true;
-		UTexture2D* Texture = CreateTexture(InOuter, AssetBasePath + AssetName, InFlattenMaterial.OpacitySize, InFlattenMaterial.OpacitySamples, TC_Default, InTextureGroup, Flags, bSRGB);
+		UTexture2D* Texture = CreateTexture(InOuter, AssetBasePath + AssetName, InFlattenMaterial.GetPropertySize(EFlattenMaterialProperties::Opacity), InFlattenMaterial.GetPropertySamples(EFlattenMaterialProperties::Opacity), TC_Default, InTextureGroup, Flags, bSRGB);
 		OutGeneratedAssets.Add(Texture);
 
 		//Assign opacity to the material
@@ -1185,10 +1197,10 @@ UMaterial* FMaterialUtilities::CreateMaterial(const FFlattenMaterial& InFlattenM
 		MaterialNodeY += MaterialNodeStepY;
 	}
 
-	if (InFlattenMaterial.SubSurfaceSamples.Num() == 1)
+	if (InFlattenMaterial.IsPropertyConstant(EFlattenMaterialProperties::SubSurface))
 	{
 		// Set Emissive to constant
-		FColor SubSurfaceColor = InFlattenMaterial.SubSurfaceSamples[0];
+		FColor SubSurfaceColor = InFlattenMaterial.GetPropertySamples(EFlattenMaterialProperties::SubSurface)[0];
 
 		// Don't have to deal with black sub surface color
 		if (SubSurfaceColor != FColor(0, 0, 0, 255))
@@ -1206,11 +1218,11 @@ UMaterial* FMaterialUtilities::CreateMaterial(const FFlattenMaterial& InFlattenM
 		Material->SetShadingModel(MSM_Subsurface);
 
 	}
-	else if (InFlattenMaterial.SubSurfaceSamples.Num() > 1)
+	else if (InFlattenMaterial.DoesPropertyContainData(EFlattenMaterialProperties::SubSurface))
 	{
 		const FString AssetName = TEXT("T_") + AssetBaseName + TEXT("_SSC");
 		const bool bSRGB = true;
-		UTexture2D* Texture = CreateTexture(InOuter, AssetBasePath + AssetName, InFlattenMaterial.SubSurfaceSize, InFlattenMaterial.SubSurfaceSamples, TC_Default, InTextureGroup, Flags, bSRGB);
+		UTexture2D* Texture = CreateTexture(InOuter, AssetBasePath + AssetName, InFlattenMaterial.GetPropertySize(EFlattenMaterialProperties::SubSurface), InFlattenMaterial.GetPropertySamples(EFlattenMaterialProperties::SubSurface), TC_Default, InTextureGroup, Flags, bSRGB);
 		OutGeneratedAssets.Add(Texture);
 
 		//Assign emissive color to the material
@@ -1345,6 +1357,7 @@ FFlattenMaterial FMaterialUtilities::CreateFlattenMaterialWithSettings(const FMa
 	// Create new material.
 	FFlattenMaterial Material;
 
+	// TODO REMOVE THIS FEATURE?
 	FIntPoint MaximumSize = InMaterialLODSettings.TextureSize;	
 	// If the user is manually overriding the texture size, make sure we have the max texture size to render with
 	if (InMaterialLODSettings.TextureSizingType == TextureSizingType_UseManualOverrideTextureSize)
@@ -1362,13 +1375,13 @@ FFlattenMaterial FMaterialUtilities::CreateFlattenMaterialWithSettings(const FMa
 	{
 		Material.RenderSize = MaximumSize;
 
-		Material.DiffuseSize = InMaterialLODSettings.DiffuseTextureSize;
-		Material.SpecularSize = (InMaterialLODSettings.bSpecularMap) ? InMaterialLODSettings.SpecularTextureSize : FIntPoint::ZeroValue;
-		Material.MetallicSize = (InMaterialLODSettings.bMetallicMap) ? InMaterialLODSettings.MetallicTextureSize : FIntPoint::ZeroValue;
-		Material.RoughnessSize = (InMaterialLODSettings.bRoughnessMap) ? InMaterialLODSettings.RoughnessTextureSize : FIntPoint::ZeroValue;
-		Material.NormalSize = (InMaterialLODSettings.bNormalMap) ? InMaterialLODSettings.NormalTextureSize : FIntPoint::ZeroValue;
-		Material.EmissiveSize = (InMaterialLODSettings.bEmissiveMap) ? InMaterialLODSettings.EmissiveTextureSize : FIntPoint::ZeroValue;
-		Material.OpacitySize = (InMaterialLODSettings.bOpacityMap) ? InMaterialLODSettings.OpacityTextureSize : FIntPoint::ZeroValue;
+		Material.SetPropertySize(EFlattenMaterialProperties::Diffuse, InMaterialLODSettings.DiffuseTextureSize);
+		Material.SetPropertySize(EFlattenMaterialProperties::Specular, InMaterialLODSettings.bSpecularMap ? InMaterialLODSettings.SpecularTextureSize : FIntPoint::ZeroValue);
+		Material.SetPropertySize(EFlattenMaterialProperties::Metallic, InMaterialLODSettings.bMetallicMap ? InMaterialLODSettings.MetallicTextureSize : FIntPoint::ZeroValue);
+		Material.SetPropertySize(EFlattenMaterialProperties::Roughness, InMaterialLODSettings.bRoughnessMap ? InMaterialLODSettings.RoughnessTextureSize : FIntPoint::ZeroValue);
+		Material.SetPropertySize(EFlattenMaterialProperties::Normal, InMaterialLODSettings.bNormalMap ? InMaterialLODSettings.NormalTextureSize : FIntPoint::ZeroValue);
+		Material.SetPropertySize(EFlattenMaterialProperties::Emissive, InMaterialLODSettings.bEmissiveMap ? InMaterialLODSettings.EmissiveTextureSize : FIntPoint::ZeroValue);
+		Material.SetPropertySize(EFlattenMaterialProperties::Opacity, InMaterialLODSettings.bOpacityMap ? InMaterialLODSettings.OpacityTextureSize : FIntPoint::ZeroValue);
 	}
 	else if (InMaterialLODSettings.TextureSizingType == TextureSizingType_UseAutomaticBiasedSizes)
 	{
@@ -1379,26 +1392,25 @@ FFlattenMaterial FMaterialUtilities::CreateFlattenMaterialWithSettings(const FMa
 		DiffuseSizeX = FMath::Max(InMaterialLODSettings.TextureSize.X >> 1, 32);
 		PropertiesSizeX = FMath::Max(InMaterialLODSettings.TextureSize.X >> 2, 16);
 
-		Material.DiffuseSize = FIntPoint(DiffuseSizeX, DiffuseSizeX);
-		Material.NormalSize = (InMaterialLODSettings.bNormalMap) ? FIntPoint(NormalSizeX, NormalSizeX) : FIntPoint::ZeroValue;
+		Material.SetPropertySize(EFlattenMaterialProperties::Diffuse, FIntPoint(DiffuseSizeX, DiffuseSizeX));
+		Material.SetPropertySize(EFlattenMaterialProperties::Normal, (InMaterialLODSettings.bNormalMap) ? FIntPoint(NormalSizeX, NormalSizeX) : FIntPoint::ZeroValue);
 
 		FIntPoint PropertiesSize = FIntPoint(PropertiesSizeX, PropertiesSizeX);
-		Material.SpecularSize = (InMaterialLODSettings.bSpecularMap) ? PropertiesSize: FIntPoint::ZeroValue;
-		Material.MetallicSize = (InMaterialLODSettings.bMetallicMap) ? PropertiesSize: FIntPoint::ZeroValue;
-		Material.RoughnessSize = (InMaterialLODSettings.bRoughnessMap) ? PropertiesSize : FIntPoint::ZeroValue;		
-		Material.EmissiveSize = (InMaterialLODSettings.bEmissiveMap) ? PropertiesSize : FIntPoint::ZeroValue;
-		Material.OpacitySize = (InMaterialLODSettings.bOpacityMap) ? PropertiesSize : FIntPoint::ZeroValue;
+		Material.SetPropertySize(EFlattenMaterialProperties::Specular, (InMaterialLODSettings.bSpecularMap) ? PropertiesSize : FIntPoint::ZeroValue );
+		Material.SetPropertySize(EFlattenMaterialProperties::Metallic, (InMaterialLODSettings.bMetallicMap) ? PropertiesSize : FIntPoint::ZeroValue );
+		Material.SetPropertySize(EFlattenMaterialProperties::Roughness, (InMaterialLODSettings.bRoughnessMap) ? PropertiesSize : FIntPoint::ZeroValue );
+		Material.SetPropertySize(EFlattenMaterialProperties::Emissive, (InMaterialLODSettings.bEmissiveMap) ? PropertiesSize : FIntPoint::ZeroValue );
+		Material.SetPropertySize(EFlattenMaterialProperties::Opacity, (InMaterialLODSettings.bOpacityMap) ? PropertiesSize : FIntPoint::ZeroValue );
 	}
 	
-	Material.RenderSize = InMaterialLODSettings.TextureSize;
-	
-	Material.DiffuseSize = InMaterialLODSettings.TextureSize;
-	Material.SpecularSize = (InMaterialLODSettings.bSpecularMap) ? InMaterialLODSettings.TextureSize : FIntPoint::ZeroValue;
-	Material.MetallicSize = (InMaterialLODSettings.bMetallicMap) ? InMaterialLODSettings.TextureSize : FIntPoint::ZeroValue;
-	Material.RoughnessSize = (InMaterialLODSettings.bRoughnessMap) ? InMaterialLODSettings.TextureSize : FIntPoint::ZeroValue;
-	Material.NormalSize = (InMaterialLODSettings.bNormalMap) ? InMaterialLODSettings.TextureSize : FIntPoint::ZeroValue;
-	Material.EmissiveSize = (InMaterialLODSettings.bEmissiveMap) ? InMaterialLODSettings.TextureSize : FIntPoint::ZeroValue;
-	Material.OpacitySize = (InMaterialLODSettings.bOpacityMap) ? InMaterialLODSettings.TextureSize : FIntPoint::ZeroValue;
+	Material.RenderSize = InMaterialLODSettings.TextureSize;	
+	Material.SetPropertySize(EFlattenMaterialProperties::Diffuse , InMaterialLODSettings.TextureSize);
+	Material.SetPropertySize(EFlattenMaterialProperties::Specular, (InMaterialLODSettings.bSpecularMap) ? InMaterialLODSettings.TextureSize : FIntPoint::ZeroValue);
+	Material.SetPropertySize(EFlattenMaterialProperties::Metallic, (InMaterialLODSettings.bMetallicMap) ? InMaterialLODSettings.TextureSize : FIntPoint::ZeroValue);
+	Material.SetPropertySize(EFlattenMaterialProperties::Roughness,  (InMaterialLODSettings.bRoughnessMap) ? InMaterialLODSettings.TextureSize : FIntPoint::ZeroValue);
+	Material.SetPropertySize(EFlattenMaterialProperties::Normal,  (InMaterialLODSettings.bNormalMap) ? InMaterialLODSettings.TextureSize : FIntPoint::ZeroValue);
+	Material.SetPropertySize(EFlattenMaterialProperties::Emissive,  (InMaterialLODSettings.bEmissiveMap) ? InMaterialLODSettings.TextureSize : FIntPoint::ZeroValue);
+	Material.SetPropertySize(EFlattenMaterialProperties::Opacity,  (InMaterialLODSettings.bOpacityMap) ? InMaterialLODSettings.TextureSize : FIntPoint::ZeroValue);
 
 	return Material;
 }
@@ -1594,19 +1606,18 @@ void FMaterialUtilities::RemapUniqueMaterialIndices(const TArray<UMaterialInterf
 void FMaterialUtilities::OptimizeFlattenMaterial(FFlattenMaterial& InFlattenMaterial)
 {
 	// Try to optimize each individual property sample
-	OptimizeSampleArray(InFlattenMaterial.DiffuseSamples, InFlattenMaterial.DiffuseSize);
-	OptimizeSampleArray(InFlattenMaterial.NormalSamples, InFlattenMaterial.NormalSize);
-	OptimizeSampleArray(InFlattenMaterial.MetallicSamples, InFlattenMaterial.MetallicSize);
-	OptimizeSampleArray(InFlattenMaterial.RoughnessSamples, InFlattenMaterial.RoughnessSize);
-	OptimizeSampleArray(InFlattenMaterial.SpecularSamples, InFlattenMaterial.SpecularSize);
-	OptimizeSampleArray(InFlattenMaterial.OpacitySamples, InFlattenMaterial.OpacitySize);
-	OptimizeSampleArray(InFlattenMaterial.EmissiveSamples, InFlattenMaterial.EmissiveSize);
+	for (int32 PropertyIndex = 0; PropertyIndex < (int32)EFlattenMaterialProperties::NumFlattenMaterialProperties; ++PropertyIndex)
+	{
+		EFlattenMaterialProperties Property = (EFlattenMaterialProperties)PropertyIndex;
+		FIntPoint Size = InFlattenMaterial.GetPropertySize(Property);
+		OptimizeSampleArray(InFlattenMaterial.GetPropertySamples(Property), Size);
+		InFlattenMaterial.SetPropertySize(Property, Size);
+	}
 }
 
 void FMaterialUtilities::ResizeFlattenMaterial(FFlattenMaterial& InFlattenMaterial, const struct FMeshProxySettings& InMeshProxySettings)
 {
 	const FMaterialProxySettings& MaterialSettings = InMeshProxySettings.MaterialSettings;
-
 	if (MaterialSettings.TextureSizingType == TextureSizingType_UseAutomaticBiasedSizes)
 	{
 		int NormalSizeX, DiffuseSizeX, PropertiesSizeX;
@@ -1614,132 +1625,132 @@ void FMaterialUtilities::ResizeFlattenMaterial(FFlattenMaterial& InFlattenMateri
 		DiffuseSizeX = FMath::Max(MaterialSettings.TextureSize.X >> 1, 32);
 		PropertiesSizeX = FMath::Max(MaterialSettings.TextureSize.X >> 2, 16);
 
-		if (InFlattenMaterial.DiffuseSize.X != DiffuseSizeX)
+		if (InFlattenMaterial.GetPropertySize(EFlattenMaterialProperties::Diffuse).X != DiffuseSizeX)
 		{
 			TArray<FColor> NewSamples;
-			FImageUtils::ImageResize(InFlattenMaterial.DiffuseSize.X, InFlattenMaterial.DiffuseSize.Y, InFlattenMaterial.DiffuseSamples, DiffuseSizeX, DiffuseSizeX, NewSamples, false);
-			InFlattenMaterial.DiffuseSamples.Reset(NewSamples.Num());
-			InFlattenMaterial.DiffuseSamples.Append(NewSamples);
-			InFlattenMaterial.DiffuseSize = FIntPoint(DiffuseSizeX, DiffuseSizeX);
+			FImageUtils::ImageResize(InFlattenMaterial.GetPropertySize(EFlattenMaterialProperties::Diffuse).X, InFlattenMaterial.GetPropertySize(EFlattenMaterialProperties::Diffuse).Y, InFlattenMaterial.GetPropertySamples(EFlattenMaterialProperties::Diffuse), DiffuseSizeX, DiffuseSizeX, NewSamples, false);
+			InFlattenMaterial.GetPropertySamples(EFlattenMaterialProperties::Diffuse).Reset(NewSamples.Num());
+			InFlattenMaterial.GetPropertySamples(EFlattenMaterialProperties::Diffuse).Append(NewSamples);
+			InFlattenMaterial.SetPropertySize(EFlattenMaterialProperties::Diffuse, FIntPoint(DiffuseSizeX, DiffuseSizeX));
 		}
 
-		if (InFlattenMaterial.SpecularSamples.Num() && MaterialSettings.bSpecularMap && InFlattenMaterial.SpecularSize.X != PropertiesSizeX)
+		if (InFlattenMaterial.GetPropertySamples(EFlattenMaterialProperties::Specular).Num() && MaterialSettings.bSpecularMap && InFlattenMaterial.GetPropertySize(EFlattenMaterialProperties::Specular).X != PropertiesSizeX)
 		{
 			TArray<FColor> NewSamples;
-			FImageUtils::ImageResize(InFlattenMaterial.SpecularSize.X, InFlattenMaterial.SpecularSize.Y, InFlattenMaterial.SpecularSamples, PropertiesSizeX, PropertiesSizeX, NewSamples, false);
-			InFlattenMaterial.SpecularSamples.Reset(NewSamples.Num());
-			InFlattenMaterial.SpecularSamples.Append(NewSamples);
-			InFlattenMaterial.SpecularSize = FIntPoint(PropertiesSizeX, PropertiesSizeX);
+			FImageUtils::ImageResize(InFlattenMaterial.GetPropertySize(EFlattenMaterialProperties::Specular).X, InFlattenMaterial.GetPropertySize(EFlattenMaterialProperties::Specular).Y, InFlattenMaterial.GetPropertySamples(EFlattenMaterialProperties::Specular), PropertiesSizeX, PropertiesSizeX, NewSamples, false);
+			InFlattenMaterial.GetPropertySamples(EFlattenMaterialProperties::Specular).Reset(NewSamples.Num());
+			InFlattenMaterial.GetPropertySamples(EFlattenMaterialProperties::Specular).Append(NewSamples);
+			InFlattenMaterial.SetPropertySize(EFlattenMaterialProperties::Specular, FIntPoint(PropertiesSizeX, PropertiesSizeX));
 		}
 
-		if (InFlattenMaterial.MetallicSamples.Num() && MaterialSettings.bMetallicMap && InFlattenMaterial.MetallicSize.X != PropertiesSizeX)
+		if (InFlattenMaterial.GetPropertySamples(EFlattenMaterialProperties::Metallic).Num() && MaterialSettings.bMetallicMap && InFlattenMaterial.GetPropertySize(EFlattenMaterialProperties::Metallic).X != PropertiesSizeX)
 		{
 			TArray<FColor> NewSamples;
-			FImageUtils::ImageResize(InFlattenMaterial.MetallicSize.X, InFlattenMaterial.MetallicSize.Y, InFlattenMaterial.MetallicSamples, PropertiesSizeX, PropertiesSizeX, NewSamples, false);
-			InFlattenMaterial.MetallicSamples.Reset(NewSamples.Num());
-			InFlattenMaterial.MetallicSamples.Append(NewSamples);
-			InFlattenMaterial.MetallicSize = FIntPoint(PropertiesSizeX, PropertiesSizeX);
+			FImageUtils::ImageResize(InFlattenMaterial.GetPropertySize(EFlattenMaterialProperties::Metallic).X, InFlattenMaterial.GetPropertySize(EFlattenMaterialProperties::Metallic).Y, InFlattenMaterial.GetPropertySamples(EFlattenMaterialProperties::Metallic), PropertiesSizeX, PropertiesSizeX, NewSamples, false);
+			InFlattenMaterial.GetPropertySamples(EFlattenMaterialProperties::Metallic).Reset(NewSamples.Num());
+			InFlattenMaterial.GetPropertySamples(EFlattenMaterialProperties::Metallic).Append(NewSamples);
+			InFlattenMaterial.SetPropertySize(EFlattenMaterialProperties::Metallic, FIntPoint(PropertiesSizeX, PropertiesSizeX));
 		}
 
-		if (InFlattenMaterial.RoughnessSamples.Num() && MaterialSettings.bRoughnessMap && InFlattenMaterial.RoughnessSize.X != PropertiesSizeX)
+		if (InFlattenMaterial.GetPropertySamples(EFlattenMaterialProperties::Roughness).Num() && MaterialSettings.bRoughnessMap && InFlattenMaterial.GetPropertySize(EFlattenMaterialProperties::Roughness).X != PropertiesSizeX)
 		{
 			TArray<FColor> NewSamples;
-			FImageUtils::ImageResize(InFlattenMaterial.RoughnessSize.X, InFlattenMaterial.RoughnessSize.Y, InFlattenMaterial.RoughnessSamples, PropertiesSizeX, PropertiesSizeX, NewSamples, false);
-			InFlattenMaterial.RoughnessSamples.Reset(NewSamples.Num());
-			InFlattenMaterial.RoughnessSamples.Append(NewSamples);
-			InFlattenMaterial.RoughnessSize = FIntPoint(PropertiesSizeX, PropertiesSizeX);
+			FImageUtils::ImageResize(InFlattenMaterial.GetPropertySize(EFlattenMaterialProperties::Roughness).X, InFlattenMaterial.GetPropertySize(EFlattenMaterialProperties::Roughness).Y, InFlattenMaterial.GetPropertySamples(EFlattenMaterialProperties::Roughness), PropertiesSizeX, PropertiesSizeX, NewSamples, false);
+			InFlattenMaterial.GetPropertySamples(EFlattenMaterialProperties::Roughness).Reset(NewSamples.Num());
+			InFlattenMaterial.GetPropertySamples(EFlattenMaterialProperties::Roughness).Append(NewSamples);
+			InFlattenMaterial.SetPropertySize(EFlattenMaterialProperties::Roughness, FIntPoint(PropertiesSizeX, PropertiesSizeX));
 		}
 
-		if (InFlattenMaterial.NormalSamples.Num() && MaterialSettings.bNormalMap && InFlattenMaterial.NormalSize.X != NormalSizeX)
+		if (InFlattenMaterial.GetPropertySamples(EFlattenMaterialProperties::Normal).Num() && MaterialSettings.bNormalMap && InFlattenMaterial.GetPropertySize(EFlattenMaterialProperties::Normal).X != NormalSizeX)
 		{
 			TArray<FColor> NewSamples;
-			FImageUtils::ImageResize(InFlattenMaterial.NormalSize.X, InFlattenMaterial.NormalSize.Y, InFlattenMaterial.NormalSamples, NormalSizeX, NormalSizeX, NewSamples, true);
-			InFlattenMaterial.NormalSamples.Reset(NewSamples.Num());
-			InFlattenMaterial.NormalSamples.Append(NewSamples);
-			InFlattenMaterial.NormalSize = FIntPoint(NormalSizeX, NormalSizeX);
+			FImageUtils::ImageResize(InFlattenMaterial.GetPropertySize(EFlattenMaterialProperties::Normal).X, InFlattenMaterial.GetPropertySize(EFlattenMaterialProperties::Normal).Y, InFlattenMaterial.GetPropertySamples(EFlattenMaterialProperties::Normal), NormalSizeX, NormalSizeX, NewSamples, true);
+			InFlattenMaterial.GetPropertySamples(EFlattenMaterialProperties::Normal).Reset(NewSamples.Num());
+			InFlattenMaterial.GetPropertySamples(EFlattenMaterialProperties::Normal).Append(NewSamples);
+			InFlattenMaterial.SetPropertySize(EFlattenMaterialProperties::Normal, FIntPoint(NormalSizeX, NormalSizeX));
 		}
 
-		if (InFlattenMaterial.EmissiveSamples.Num() && MaterialSettings.bEmissiveMap && InFlattenMaterial.EmissiveSize.X != PropertiesSizeX)
+		if (InFlattenMaterial.GetPropertySamples(EFlattenMaterialProperties::Emissive).Num() && MaterialSettings.bEmissiveMap && InFlattenMaterial.GetPropertySize(EFlattenMaterialProperties::Emissive).X != PropertiesSizeX)
 		{
 			TArray<FColor> NewSamples;
-			FImageUtils::ImageResize(InFlattenMaterial.EmissiveSize.X, InFlattenMaterial.EmissiveSize.Y, InFlattenMaterial.EmissiveSamples, PropertiesSizeX, PropertiesSizeX, NewSamples, false);
-			InFlattenMaterial.EmissiveSamples.Reset(NewSamples.Num());
-			InFlattenMaterial.EmissiveSamples.Append(NewSamples);
-			InFlattenMaterial.EmissiveSize = FIntPoint(PropertiesSizeX, PropertiesSizeX);
+			FImageUtils::ImageResize(InFlattenMaterial.GetPropertySize(EFlattenMaterialProperties::Emissive).X, InFlattenMaterial.GetPropertySize(EFlattenMaterialProperties::Emissive).Y, InFlattenMaterial.GetPropertySamples(EFlattenMaterialProperties::Emissive), PropertiesSizeX, PropertiesSizeX, NewSamples, false);
+			InFlattenMaterial.GetPropertySamples(EFlattenMaterialProperties::Emissive).Reset(NewSamples.Num());
+			InFlattenMaterial.GetPropertySamples(EFlattenMaterialProperties::Emissive).Append(NewSamples);
+			InFlattenMaterial.SetPropertySize(EFlattenMaterialProperties::Emissive, FIntPoint(PropertiesSizeX, PropertiesSizeX));
 		}
 
-		if (InFlattenMaterial.OpacitySamples.Num() && MaterialSettings.bOpacityMap && InFlattenMaterial.OpacitySize.X != PropertiesSizeX)
+		if (InFlattenMaterial.GetPropertySamples(EFlattenMaterialProperties::Opacity).Num() && MaterialSettings.bOpacityMap && InFlattenMaterial.GetPropertySize(EFlattenMaterialProperties::Opacity).X != PropertiesSizeX)
 		{
 			TArray<FColor> NewSamples;
-			FImageUtils::ImageResize(InFlattenMaterial.OpacitySize.X, InFlattenMaterial.OpacitySize.Y, InFlattenMaterial.OpacitySamples, PropertiesSizeX, PropertiesSizeX, NewSamples, false);
-			InFlattenMaterial.OpacitySamples.Reset(NewSamples.Num());
-			InFlattenMaterial.OpacitySamples.Append(NewSamples);
-			InFlattenMaterial.OpacitySize = FIntPoint(PropertiesSizeX, PropertiesSizeX);
+			FImageUtils::ImageResize(InFlattenMaterial.GetPropertySize(EFlattenMaterialProperties::Opacity).X, InFlattenMaterial.GetPropertySize(EFlattenMaterialProperties::Opacity).Y, InFlattenMaterial.GetPropertySamples(EFlattenMaterialProperties::Opacity), PropertiesSizeX, PropertiesSizeX, NewSamples, false);
+			InFlattenMaterial.GetPropertySamples(EFlattenMaterialProperties::Opacity).Reset(NewSamples.Num());
+			InFlattenMaterial.GetPropertySamples(EFlattenMaterialProperties::Opacity).Append(NewSamples);
+			InFlattenMaterial.SetPropertySize(EFlattenMaterialProperties::Opacity, FIntPoint(PropertiesSizeX, PropertiesSizeX));
 		}
 	}
 	else if (MaterialSettings.TextureSizingType == TextureSizingType_UseManualOverrideTextureSize)
 	{
-		if (InFlattenMaterial.DiffuseSize != MaterialSettings.DiffuseTextureSize)
+		if (InFlattenMaterial.GetPropertySize(EFlattenMaterialProperties::Diffuse) != MaterialSettings.DiffuseTextureSize)
 		{
 			TArray<FColor> NewSamples;
-			FImageUtils::ImageResize(InFlattenMaterial.DiffuseSize.X, InFlattenMaterial.DiffuseSize.Y, InFlattenMaterial.DiffuseSamples, MaterialSettings.DiffuseTextureSize.X, MaterialSettings.DiffuseTextureSize.Y, NewSamples, false);
-			InFlattenMaterial.DiffuseSamples.Reset(NewSamples.Num());
-			InFlattenMaterial.DiffuseSamples.Append(NewSamples);
-			InFlattenMaterial.DiffuseSize = MaterialSettings.DiffuseTextureSize;
+			FImageUtils::ImageResize(InFlattenMaterial.GetPropertySize(EFlattenMaterialProperties::Diffuse).X, InFlattenMaterial.GetPropertySize(EFlattenMaterialProperties::Diffuse).Y, InFlattenMaterial.GetPropertySamples(EFlattenMaterialProperties::Diffuse), MaterialSettings.DiffuseTextureSize.X, MaterialSettings.DiffuseTextureSize.Y, NewSamples, false);
+			InFlattenMaterial.GetPropertySamples(EFlattenMaterialProperties::Diffuse).Reset(NewSamples.Num());
+			InFlattenMaterial.GetPropertySamples(EFlattenMaterialProperties::Diffuse).Append(NewSamples);
+			InFlattenMaterial.SetPropertySize(EFlattenMaterialProperties::Diffuse, MaterialSettings.DiffuseTextureSize);
 		}
 
-		if (InFlattenMaterial.SpecularSamples.Num() && MaterialSettings.bSpecularMap && InFlattenMaterial.SpecularSize != MaterialSettings.SpecularTextureSize)
+		if (InFlattenMaterial.GetPropertySamples(EFlattenMaterialProperties::Specular).Num() && MaterialSettings.bSpecularMap && InFlattenMaterial.GetPropertySize(EFlattenMaterialProperties::Specular) != MaterialSettings.SpecularTextureSize)
 		{
 			TArray<FColor> NewSamples;
-			FImageUtils::ImageResize(InFlattenMaterial.SpecularSize.X, InFlattenMaterial.SpecularSize.Y, InFlattenMaterial.SpecularSamples, MaterialSettings.SpecularTextureSize.X, MaterialSettings.SpecularTextureSize.Y, NewSamples, false);
-			InFlattenMaterial.SpecularSamples.Reset(NewSamples.Num());
-			InFlattenMaterial.SpecularSamples.Append(NewSamples);
-			InFlattenMaterial.SpecularSize = MaterialSettings.SpecularTextureSize;
+			FImageUtils::ImageResize(InFlattenMaterial.GetPropertySize(EFlattenMaterialProperties::Specular).X, InFlattenMaterial.GetPropertySize(EFlattenMaterialProperties::Specular).Y, InFlattenMaterial.GetPropertySamples(EFlattenMaterialProperties::Specular), MaterialSettings.SpecularTextureSize.X, MaterialSettings.SpecularTextureSize.Y, NewSamples, false);
+			InFlattenMaterial.GetPropertySamples(EFlattenMaterialProperties::Specular).Reset(NewSamples.Num());
+			InFlattenMaterial.GetPropertySamples(EFlattenMaterialProperties::Specular).Append(NewSamples);
+			InFlattenMaterial.SetPropertySize(EFlattenMaterialProperties::Specular, MaterialSettings.SpecularTextureSize);
 		}
 
-		if (InFlattenMaterial.MetallicSamples.Num() && MaterialSettings.bMetallicMap && InFlattenMaterial.MetallicSize != MaterialSettings.MetallicTextureSize)
+		if (InFlattenMaterial.GetPropertySamples(EFlattenMaterialProperties::Metallic).Num() && MaterialSettings.bMetallicMap && InFlattenMaterial.GetPropertySize(EFlattenMaterialProperties::Metallic) != MaterialSettings.MetallicTextureSize)
 		{
 			TArray<FColor> NewSamples;
-			FImageUtils::ImageResize(InFlattenMaterial.MetallicSize.X, InFlattenMaterial.MetallicSize.Y, InFlattenMaterial.MetallicSamples, MaterialSettings.MetallicTextureSize.X, MaterialSettings.MetallicTextureSize.Y, NewSamples, false);
-			InFlattenMaterial.MetallicSamples.Reset(NewSamples.Num());
-			InFlattenMaterial.MetallicSamples.Append(NewSamples);
-			InFlattenMaterial.MetallicSize = MaterialSettings.MetallicTextureSize;
+			FImageUtils::ImageResize(InFlattenMaterial.GetPropertySize(EFlattenMaterialProperties::Metallic).X, InFlattenMaterial.GetPropertySize(EFlattenMaterialProperties::Metallic).Y, InFlattenMaterial.GetPropertySamples(EFlattenMaterialProperties::Metallic), MaterialSettings.MetallicTextureSize.X, MaterialSettings.MetallicTextureSize.Y, NewSamples, false);
+			InFlattenMaterial.GetPropertySamples(EFlattenMaterialProperties::Metallic).Reset(NewSamples.Num());
+			InFlattenMaterial.GetPropertySamples(EFlattenMaterialProperties::Metallic).Append(NewSamples);
+			InFlattenMaterial.SetPropertySize(EFlattenMaterialProperties::Metallic, MaterialSettings.MetallicTextureSize);
 		}
 
-		if (InFlattenMaterial.RoughnessSamples.Num() && MaterialSettings.bRoughnessMap && InFlattenMaterial.RoughnessSize != MaterialSettings.RoughnessTextureSize)
+		if (InFlattenMaterial.GetPropertySamples(EFlattenMaterialProperties::Roughness).Num() && MaterialSettings.bRoughnessMap && InFlattenMaterial.GetPropertySize(EFlattenMaterialProperties::Roughness) != MaterialSettings.RoughnessTextureSize)
 		{
 			TArray<FColor> NewSamples;
-			FImageUtils::ImageResize(InFlattenMaterial.RoughnessSize.X, InFlattenMaterial.RoughnessSize.Y, InFlattenMaterial.RoughnessSamples, MaterialSettings.RoughnessTextureSize.X, MaterialSettings.RoughnessTextureSize.Y, NewSamples, false);
-			InFlattenMaterial.RoughnessSamples.Reset(NewSamples.Num());
-			InFlattenMaterial.RoughnessSamples.Append(NewSamples);
-			InFlattenMaterial.RoughnessSize = MaterialSettings.RoughnessTextureSize;
+			FImageUtils::ImageResize(InFlattenMaterial.GetPropertySize(EFlattenMaterialProperties::Roughness).X, InFlattenMaterial.GetPropertySize(EFlattenMaterialProperties::Roughness).Y, InFlattenMaterial.GetPropertySamples(EFlattenMaterialProperties::Roughness), MaterialSettings.RoughnessTextureSize.X, MaterialSettings.RoughnessTextureSize.Y, NewSamples, false);
+			InFlattenMaterial.GetPropertySamples(EFlattenMaterialProperties::Roughness).Reset(NewSamples.Num());
+			InFlattenMaterial.GetPropertySamples(EFlattenMaterialProperties::Roughness).Append(NewSamples);
+			InFlattenMaterial.SetPropertySize(EFlattenMaterialProperties::Roughness, MaterialSettings.RoughnessTextureSize);
 		}
 
-		if (InFlattenMaterial.NormalSamples.Num() && MaterialSettings.bNormalMap && InFlattenMaterial.NormalSize != MaterialSettings.NormalTextureSize)
+		if (InFlattenMaterial.GetPropertySamples(EFlattenMaterialProperties::Normal).Num() && MaterialSettings.bNormalMap && InFlattenMaterial.GetPropertySize(EFlattenMaterialProperties::Normal) != MaterialSettings.NormalTextureSize)
 		{
 			TArray<FColor> NewSamples;
-			FImageUtils::ImageResize(InFlattenMaterial.NormalSize.X, InFlattenMaterial.NormalSize.Y, InFlattenMaterial.NormalSamples, MaterialSettings.NormalTextureSize.X, MaterialSettings.NormalTextureSize.Y, NewSamples, true);
-			InFlattenMaterial.NormalSamples.Reset(NewSamples.Num());
-			InFlattenMaterial.NormalSamples.Append(NewSamples);
-			InFlattenMaterial.NormalSize = MaterialSettings.NormalTextureSize;
+			FImageUtils::ImageResize(InFlattenMaterial.GetPropertySize(EFlattenMaterialProperties::Normal).X, InFlattenMaterial.GetPropertySize(EFlattenMaterialProperties::Normal).Y, InFlattenMaterial.GetPropertySamples(EFlattenMaterialProperties::Normal), MaterialSettings.NormalTextureSize.X, MaterialSettings.NormalTextureSize.Y, NewSamples, true);
+			InFlattenMaterial.GetPropertySamples(EFlattenMaterialProperties::Normal).Reset(NewSamples.Num());
+			InFlattenMaterial.GetPropertySamples(EFlattenMaterialProperties::Normal).Append(NewSamples);
+			InFlattenMaterial.SetPropertySize(EFlattenMaterialProperties::Normal, MaterialSettings.NormalTextureSize);
 		}
 
-		if (InFlattenMaterial.EmissiveSamples.Num() && MaterialSettings.bEmissiveMap && InFlattenMaterial.EmissiveSize != MaterialSettings.EmissiveTextureSize)
+		if (InFlattenMaterial.GetPropertySamples(EFlattenMaterialProperties::Emissive).Num() && MaterialSettings.bEmissiveMap && InFlattenMaterial.GetPropertySize(EFlattenMaterialProperties::Emissive) != MaterialSettings.EmissiveTextureSize)
 		{
 			TArray<FColor> NewSamples;
-			FImageUtils::ImageResize(InFlattenMaterial.EmissiveSize.X, InFlattenMaterial.EmissiveSize.Y, InFlattenMaterial.EmissiveSamples, MaterialSettings.EmissiveTextureSize.X, MaterialSettings.EmissiveTextureSize.Y, NewSamples, false);
-			InFlattenMaterial.EmissiveSamples.Reset(NewSamples.Num());
-			InFlattenMaterial.EmissiveSamples.Append(NewSamples);
-			InFlattenMaterial.EmissiveSize = MaterialSettings.EmissiveTextureSize;
+			FImageUtils::ImageResize(InFlattenMaterial.GetPropertySize(EFlattenMaterialProperties::Emissive).X, InFlattenMaterial.GetPropertySize(EFlattenMaterialProperties::Emissive).Y, InFlattenMaterial.GetPropertySamples(EFlattenMaterialProperties::Emissive), MaterialSettings.EmissiveTextureSize.X, MaterialSettings.EmissiveTextureSize.Y, NewSamples, false);
+			InFlattenMaterial.GetPropertySamples(EFlattenMaterialProperties::Emissive).Reset(NewSamples.Num());
+			InFlattenMaterial.GetPropertySamples(EFlattenMaterialProperties::Emissive).Append(NewSamples);
+			InFlattenMaterial.SetPropertySize(EFlattenMaterialProperties::Emissive, MaterialSettings.EmissiveTextureSize);
 		}
 
-		if (InFlattenMaterial.OpacitySamples.Num() && MaterialSettings.bOpacityMap && InFlattenMaterial.OpacitySize != MaterialSettings.OpacityTextureSize)
+		if (InFlattenMaterial.GetPropertySamples(EFlattenMaterialProperties::Opacity).Num() && MaterialSettings.bOpacityMap && InFlattenMaterial.GetPropertySize(EFlattenMaterialProperties::Opacity) != MaterialSettings.OpacityTextureSize)
 		{
 			TArray<FColor> NewSamples;
-			FImageUtils::ImageResize(InFlattenMaterial.OpacitySize.X, InFlattenMaterial.OpacitySize.Y, InFlattenMaterial.OpacitySamples, MaterialSettings.OpacityTextureSize.X, MaterialSettings.OpacityTextureSize.Y, NewSamples, false);
-			InFlattenMaterial.OpacitySamples.Reset(NewSamples.Num());
-			InFlattenMaterial.OpacitySamples.Append(NewSamples);
-			InFlattenMaterial.OpacitySize = MaterialSettings.OpacityTextureSize;
+			FImageUtils::ImageResize(InFlattenMaterial.GetPropertySize(EFlattenMaterialProperties::Opacity).X, InFlattenMaterial.GetPropertySize(EFlattenMaterialProperties::Opacity).Y, InFlattenMaterial.GetPropertySamples(EFlattenMaterialProperties::Opacity), MaterialSettings.OpacityTextureSize.X, MaterialSettings.OpacityTextureSize.Y, NewSamples, false);
+			InFlattenMaterial.GetPropertySamples(EFlattenMaterialProperties::Opacity).Reset(NewSamples.Num());
+			InFlattenMaterial.GetPropertySamples(EFlattenMaterialProperties::Opacity).Append(NewSamples);
+			InFlattenMaterial.SetPropertySize(EFlattenMaterialProperties::Opacity, MaterialSettings.OpacityTextureSize);
 		}
 	}
 }
@@ -2099,51 +2110,68 @@ bool FMaterialUtilities::ExportMaterial(struct FMaterialMergeData& InMaterialDat
 	}
 
 	// Determine whether or not certain properties can be rendered
-	const bool bRenderNormal = (Material->GetMaterial()->HasNormalConnected() || Material->GetMaterial()->bUseMaterialAttributes) && (OutFlattenMaterial.NormalSize.X > 0 && OutFlattenMaterial.NormalSize.Y > 0);
-	const bool bRenderEmissive = (Material->GetMaterial()->EmissiveColor.IsConnected() || Material->GetMaterial()->bUseMaterialAttributes) && OutFlattenMaterial.EmissiveSize.X > 0 && OutFlattenMaterial.EmissiveSize.Y > 0;
-	const bool bRenderOpacityMask = Material->IsPropertyActive(MP_OpacityMask) && Material->GetBlendMode() == BLEND_Masked && OutFlattenMaterial.OpacitySize.X > 0 && OutFlattenMaterial.OpacitySize.Y > 0;
-	const bool bRenderOpacity = Material->IsPropertyActive(MP_Opacity) && IsTranslucentBlendMode(Material->GetBlendMode()) && OutFlattenMaterial.OpacitySize.X > 0 && OutFlattenMaterial.OpacitySize.Y > 0;
-	const bool bRenderSubSurface = Material->IsPropertyActive(MP_SubsurfaceColor) && OutFlattenMaterial.SubSurfaceSize.X > 0 && OutFlattenMaterial.SubSurfaceSize.Y > 0;
-	const bool bRenderMetallic = Material->IsPropertyActive(MP_Metallic) && OutFlattenMaterial.MetallicSize.X > 0 && OutFlattenMaterial.MetallicSize.Y > 0;
-	const bool bRenderSpecular = Material->IsPropertyActive(MP_Specular) && OutFlattenMaterial.SpecularSize.X > 0 && OutFlattenMaterial.SpecularSize.Y > 0;
-	const bool bRenderRoughness = Material->IsPropertyActive(MP_Roughness) && OutFlattenMaterial.RoughnessSize.X > 0 && OutFlattenMaterial.RoughnessSize.Y > 0;
+	const bool bRenderNormal = (Material->GetMaterial()->HasNormalConnected() || Material->GetMaterial()->bUseMaterialAttributes) && OutFlattenMaterial.ShouldGenerateDataForProperty(EFlattenMaterialProperties::Normal);
+	const bool bRenderEmissive = (Material->GetMaterial()->EmissiveColor.IsConnected() || Material->GetMaterial()->bUseMaterialAttributes) && OutFlattenMaterial.ShouldGenerateDataForProperty(EFlattenMaterialProperties::Emissive);
+	const bool bRenderOpacityMask = Material->IsPropertyActive(MP_OpacityMask) && Material->GetBlendMode() == BLEND_Masked && OutFlattenMaterial.ShouldGenerateDataForProperty(EFlattenMaterialProperties::Opacity);
+	const bool bRenderOpacity = Material->IsPropertyActive(MP_Opacity) && IsTranslucentBlendMode(Material->GetBlendMode()) && OutFlattenMaterial.ShouldGenerateDataForProperty(EFlattenMaterialProperties::Opacity);
+	const bool bRenderSubSurface = Material->IsPropertyActive(MP_SubsurfaceColor) && OutFlattenMaterial.ShouldGenerateDataForProperty(EFlattenMaterialProperties::SubSurface);
+	const bool bRenderMetallic = Material->IsPropertyActive(MP_Metallic) && OutFlattenMaterial.ShouldGenerateDataForProperty(EFlattenMaterialProperties::Metallic);
+	const bool bRenderSpecular = Material->IsPropertyActive(MP_Specular) && OutFlattenMaterial.ShouldGenerateDataForProperty(EFlattenMaterialProperties::Specular);
+	const bool bRenderRoughness = Material->IsPropertyActive(MP_Roughness) && OutFlattenMaterial.ShouldGenerateDataForProperty(EFlattenMaterialProperties::Roughness);
 
 	check(!bRenderOpacity || !bRenderOpacityMask);
 
-	// Compile shaders and render flatten material.
-	RenderMaterialPropertyToTexture(InMaterialData, MP_BaseColor, false, PF_B8G8R8A8, OutFlattenMaterial.RenderSize, OutFlattenMaterial.DiffuseSize, OutFlattenMaterial.DiffuseSamples);
+	FIntPoint Size;
+	// Compile shaders and render flatten material.	
+	Size = OutFlattenMaterial.GetPropertySize(EFlattenMaterialProperties::Diffuse);
+	RenderMaterialPropertyToTexture(InMaterialData, MP_BaseColor, false, PF_B8G8R8A8, OutFlattenMaterial.RenderSize, Size, OutFlattenMaterial.GetPropertySamples(EFlattenMaterialProperties::Diffuse));
+	OutFlattenMaterial.SetPropertySize(EFlattenMaterialProperties::Diffuse, Size);
 
 	if (bRenderMetallic)
 	{
-		RenderMaterialPropertyToTexture(InMaterialData, MP_Metallic, false, PF_B8G8R8A8, OutFlattenMaterial.RenderSize, OutFlattenMaterial.MetallicSize, OutFlattenMaterial.MetallicSamples);
-	}	
+		Size = OutFlattenMaterial.GetPropertySize(EFlattenMaterialProperties::Metallic);
+		RenderMaterialPropertyToTexture(InMaterialData, MP_Metallic, false, PF_B8G8R8A8, OutFlattenMaterial.RenderSize, Size, OutFlattenMaterial.GetPropertySamples(EFlattenMaterialProperties::Metallic));
+		OutFlattenMaterial.SetPropertySize(EFlattenMaterialProperties::Metallic, Size);
+	}
 	if (bRenderSpecular)
 	{
-		RenderMaterialPropertyToTexture(InMaterialData, MP_Specular, false, PF_B8G8R8A8, OutFlattenMaterial.RenderSize, OutFlattenMaterial.SpecularSize, OutFlattenMaterial.SpecularSamples);
+		Size = OutFlattenMaterial.GetPropertySize(EFlattenMaterialProperties::Specular);
+		RenderMaterialPropertyToTexture(InMaterialData, MP_Specular, false, PF_B8G8R8A8, OutFlattenMaterial.RenderSize, Size, OutFlattenMaterial.GetPropertySamples(EFlattenMaterialProperties::Specular));
+		OutFlattenMaterial.SetPropertySize(EFlattenMaterialProperties::Specular, Size);
 	}
 	if (bRenderRoughness)
 	{
-		RenderMaterialPropertyToTexture(InMaterialData, MP_Roughness, false, PF_B8G8R8A8, OutFlattenMaterial.RenderSize, OutFlattenMaterial.RoughnessSize, OutFlattenMaterial.RoughnessSamples);
+		Size = OutFlattenMaterial.GetPropertySize(EFlattenMaterialProperties::Roughness);
+		RenderMaterialPropertyToTexture(InMaterialData, MP_Roughness, false, PF_B8G8R8A8, OutFlattenMaterial.RenderSize, Size, OutFlattenMaterial.GetPropertySamples(EFlattenMaterialProperties::Roughness));
+		OutFlattenMaterial.SetPropertySize(EFlattenMaterialProperties::Roughness, Size);
 	}
 	if (bRenderNormal)
 	{
-		RenderMaterialPropertyToTexture(InMaterialData, MP_Normal, true, PF_B8G8R8A8, OutFlattenMaterial.RenderSize, OutFlattenMaterial.NormalSize, OutFlattenMaterial.NormalSamples);
+		Size = OutFlattenMaterial.GetPropertySize(EFlattenMaterialProperties::Normal);
+		RenderMaterialPropertyToTexture(InMaterialData, MP_Normal, true, PF_B8G8R8A8, OutFlattenMaterial.RenderSize, Size, OutFlattenMaterial.GetPropertySamples(EFlattenMaterialProperties::Normal));
+		OutFlattenMaterial.SetPropertySize(EFlattenMaterialProperties::Normal, Size);
 	}
 	if (bRenderOpacityMask)
 	{
-		RenderMaterialPropertyToTexture(InMaterialData, MP_OpacityMask, true, PF_B8G8R8A8, OutFlattenMaterial.RenderSize, OutFlattenMaterial.OpacitySize, OutFlattenMaterial.OpacitySamples);
+		Size = OutFlattenMaterial.GetPropertySize(EFlattenMaterialProperties::Opacity);
+		RenderMaterialPropertyToTexture(InMaterialData, MP_OpacityMask, true, PF_B8G8R8A8, OutFlattenMaterial.RenderSize, Size, OutFlattenMaterial.GetPropertySamples(EFlattenMaterialProperties::Opacity));
+		OutFlattenMaterial.SetPropertySize(EFlattenMaterialProperties::Opacity, Size);
 	}
 	if (bRenderOpacity)
 	{
+		Size = OutFlattenMaterial.GetPropertySize(EFlattenMaterialProperties::Opacity);
 		// Number of blend modes, let's UMaterial decide whether it wants this property
-		RenderMaterialPropertyToTexture(InMaterialData, MP_Opacity, true, PF_B8G8R8A8, OutFlattenMaterial.RenderSize, OutFlattenMaterial.OpacitySize, OutFlattenMaterial.OpacitySamples);
+		RenderMaterialPropertyToTexture(InMaterialData, MP_Opacity, true, PF_B8G8R8A8, OutFlattenMaterial.RenderSize, Size, OutFlattenMaterial.GetPropertySamples(EFlattenMaterialProperties::Opacity));
+		OutFlattenMaterial.SetPropertySize(EFlattenMaterialProperties::Opacity, Size);
 	}
 	if (bRenderEmissive)
 	{
+		Size = OutFlattenMaterial.GetPropertySize(EFlattenMaterialProperties::Emissive);
 		// PF_FloatRGBA is here to be able to render and read HDR image using ReadFloat16Pixels()
-		RenderMaterialPropertyToTexture(InMaterialData, MP_EmissiveColor, false, PF_FloatRGBA, OutFlattenMaterial.RenderSize, OutFlattenMaterial.EmissiveSize, OutFlattenMaterial.EmissiveSamples);
+		RenderMaterialPropertyToTexture(InMaterialData, MP_EmissiveColor, false, PF_FloatRGBA, OutFlattenMaterial.RenderSize, Size, OutFlattenMaterial.GetPropertySamples(EFlattenMaterialProperties::Emissive));
 		OutFlattenMaterial.EmissiveScale = InMaterialData.EmissiveScale;
-	}	
+		OutFlattenMaterial.SetPropertySize(EFlattenMaterialProperties::Emissive, Size);
+	}
 
 	if (bRenderSubSurface)
 	{
@@ -2165,7 +2193,7 @@ bool FMaterialUtilities::ExportMaterial(struct FMaterialMergeData& InMaterialDat
 	return true;
 }
 
-bool FMaterialUtilities::RenderMaterialPropertyToTexture(struct FMaterialMergeData& InMaterialData, EMaterialProperty InMaterialProperty, bool bInForceLinearGamma, EPixelFormat InPixelFormat, const FIntPoint& InTargetSize, FIntPoint& OutSampleSize, TArray<FColor>& OutSamples)
+bool FMaterialUtilities::RenderMaterialPropertyToTexture(struct FMaterialMergeData& InMaterialData, EMaterialProperty InMaterialProperty, bool bInForceLinearGamma, EPixelFormat InPixelFormat, const FIntPoint InTargetSize, FIntPoint& OutSampleSize, TArray<FColor>& OutSamples)
 {
 	if (InTargetSize.X == 0 || InTargetSize.Y == 0)
 	{
@@ -2370,3 +2398,193 @@ void FExportMaterialProxyCache::Release()
 	}
 }
 
+
+void FMaterialUtilities::DetermineMaterialImportance(const TArray<UMaterialInterface*>& InMaterials, TArray<float>& OutImportance)
+{
+	TArray<int32> MaterialImportance;
+	int32 SummedSize = 0;
+	for (UMaterialInterface* Material : InMaterials)
+	{
+		TArray<UTexture*> UsedTextures;
+		Material->GetMaterial()->AppendReferencedTextures(UsedTextures);
+		if (UMaterialInstance* MaterialInstance = Cast<UMaterialInstance>(Material))
+		{
+			for (const FTextureParameterValue& TextureParameter : MaterialInstance->TextureParameterValues)
+			{
+				if (TextureParameter.ParameterValue != nullptr)
+				{
+					UsedTextures.Add(TextureParameter.ParameterValue);
+				}
+			}
+		}
+		int32 MaxSize = 64 * 64;
+		for (UTexture* Texture : UsedTextures)
+		{
+			if (UTexture2D* Texture2D = Cast<UTexture2D>(Texture))
+			{
+
+				// In game max bias and dimensions
+				uint32 MaxInGameWidth, MaxInGameHeight;
+				int32 MipLevel = UDeviceProfileManager::Get().GetActiveProfile()->GetTextureLODSettings()->CalculateLODBias(Texture2D);
+				//Calculate in-game max resolution and store in EffectiveTextureWidth, EffectiveTextureHeight
+				UDeviceProfileManager::Get().GetActiveProfile()->GetTextureLODSettings()->ComputeInGameMaxResolution(MipLevel, *Texture, (uint32 &)MaxInGameWidth, (uint32 &)MaxInGameHeight);
+				
+				int32 TotalSize = MaxInGameWidth * MaxInGameHeight;
+				if (TotalSize > MaxSize)
+				{
+					MaxSize = TotalSize;
+				}
+			}
+		}
+		
+		MaterialImportance.Add(MaxSize);
+		SummedSize += MaxSize;
+	}
+	float WeightPerPixel = 1.0f / SummedSize;
+	for (int32 MaterialIndex = 0; MaterialIndex < InMaterials.Num(); ++MaterialIndex)
+	{
+		OutImportance.Add((float)MaterialImportance[MaterialIndex] * WeightPerPixel);
+	}	
+}
+
+void FMaterialUtilities::GeneratedBinnedTextureSquares(const FVector2D DestinationSize, TArray<float>& InTexureWeights, TArray<FBox2D>& OutGeneratedBoxes)
+{
+	typedef FBox2D FTextureArea;
+	struct FWeightedTexture
+	{
+		FTextureArea Area;
+		int32 TextureIndex;
+		float Weight;
+	};
+
+	TArray<FWeightedTexture> WeightedTextures;
+	const float TotalArea = DestinationSize.X * DestinationSize.Y;
+	// Generate textures with their size calculated according to their weight
+	for (int32 WeightIndex = 0; WeightIndex < InTexureWeights.Num(); ++WeightIndex)
+	{
+		const float Weight = InTexureWeights[WeightIndex];
+		FWeightedTexture Texture;
+		float TextureSize = FMath::Sqrt(TotalArea*Weight);
+		Texture.Area = FTextureArea( FVector2D(0.0f,0.0f), FVector2D(TextureSize, TextureSize));
+		Texture.TextureIndex = WeightIndex;
+		Texture.Weight = Weight;
+		WeightedTextures.Add(Texture);
+	}
+
+	// Sort textures by their weight (high to low) which influences the insert order
+	WeightedTextures.Sort([](const FWeightedTexture& One, const FWeightedTexture& Two) { return One.Weight > Two.Weight; });
+		
+	TArray<FWeightedTexture> InsertedTextures;
+	typedef FBox2D FUnusedArea;
+	TArray<FUnusedArea> UnusedAreas;
+
+	bool bSuccess = true;
+	do 
+	{
+		// Reset state
+		bSuccess = true;
+		UnusedAreas.Empty();
+		InsertedTextures.Empty();
+		FUnusedArea StartArea(FVector2D(0, 0), DestinationSize);
+		UnusedAreas.Add(StartArea);
+
+		for (const FWeightedTexture& Texture : WeightedTextures)
+		{
+			int32 BestAreaIndex = -1;
+			float RemainingArea = FLT_MAX;
+			FVector2D TextureSize = Texture.Area.GetSize();
+			float TextureSurface = TextureSize.X * TextureSize.Y;
+
+			// Find best area to insert this texture in (determined by tightest fit)
+			for (int32 AreaIndex = 0; AreaIndex < UnusedAreas.Num(); ++AreaIndex)
+			{
+				const FUnusedArea& UnusedArea = UnusedAreas[AreaIndex];
+				if (UnusedArea.GetSize() >= TextureSize)
+				{
+					const float Remainder = UnusedArea.GetArea() - TextureSurface;
+					if (Remainder < RemainingArea && Remainder >= 0)
+					{
+						BestAreaIndex = AreaIndex;
+						RemainingArea = Remainder;
+					}
+				}
+			}
+
+			// Insert the texture in case we found an appropriate area
+			if (BestAreaIndex != -1)
+			{
+				FUnusedArea& UnusedArea = UnusedAreas[BestAreaIndex];
+				FVector2D UnusedSize = UnusedArea.GetSize();
+
+				// Push back texture
+				FWeightedTexture WeightedTexture;
+				WeightedTexture.Area = FTextureArea(UnusedArea.Min, UnusedArea.Min + TextureSize);
+				WeightedTexture.TextureIndex = Texture.TextureIndex;
+				InsertedTextures.Add(WeightedTexture);
+				
+				// Generate two new resulting unused areas from splitting up the result
+				/*
+					___________
+					|	  |   |
+					|	  | V |
+					|_____|   |
+					|  H  |   |
+					|_____|___|
+				*/
+				FUnusedArea HorizontalArea, VerticalArea;
+				HorizontalArea.Min.X = UnusedArea.Min.X;
+				HorizontalArea.Min.Y = UnusedArea.Min.Y + TextureSize.Y;
+				HorizontalArea.Max.X = HorizontalArea.Min.X + TextureSize.X;
+				HorizontalArea.Max.Y = HorizontalArea.Min.Y + (UnusedSize.Y - TextureSize.Y);
+
+				VerticalArea.Min.X = UnusedArea.Min.X + TextureSize.X;
+				VerticalArea.Min.Y = UnusedArea.Min.Y;
+				VerticalArea.Max.X = VerticalArea.Min.X + (UnusedSize.X - TextureSize.X);
+				VerticalArea.Max.Y = UnusedSize.Y;
+
+				// Append valid new areas to list (replace original one with either one of the new ones)
+				const bool bValidHorizontal = HorizontalArea.GetArea() > 0.0f;
+				const bool bValidVertical = VerticalArea.GetArea() > 0.0f;
+				if (bValidVertical && bValidHorizontal)
+				{
+					UnusedAreas[BestAreaIndex] = HorizontalArea;
+					UnusedAreas.Add(VerticalArea);
+				}
+				else if (bValidVertical)
+				{
+					UnusedAreas[BestAreaIndex] = VerticalArea;
+				}
+				else if (bValidHorizontal)
+				{
+					UnusedAreas[BestAreaIndex] = HorizontalArea;
+				}
+			}
+			else
+			{
+				bSuccess = false;
+				break;
+			}
+		}
+
+		// This means we failed to find a fit, in this case we resize the textures and try again until we find one
+		if (bSuccess == false)
+		{
+			for (FWeightedTexture& Texture : WeightedTextures)
+			{
+				Texture.Area.Max *= .99f;
+			}
+		}
+	} while (!bSuccess);
+
+	// Now generate boxes
+	OutGeneratedBoxes.Empty(InTexureWeights.Num());
+	OutGeneratedBoxes.AddZeroed(InTexureWeights.Num());
+
+	// Generate boxes according to the inserted textures
+	for (const FWeightedTexture& Texture : InsertedTextures)
+	{
+		FBox2D& Box = OutGeneratedBoxes[Texture.TextureIndex];
+		Box = Texture.Area;
+	}
+}
+
diff --git a/Engine/Source/Developer/MaterialUtilities/Public/MaterialUtilities.h b/Engine/Source/Developer/MaterialUtilities/Public/MaterialUtilities.h
--- a/Engine/Source/Developer/MaterialUtilities/Public/MaterialUtilities.h
+++ b/Engine/Source/Developer/MaterialUtilities/Public/MaterialUtilities.h
@@ -11,98 +11,73 @@
 #include "Runtime/Engine/Classes/Engine/TextureStreamingTypes.h"
 #include "RawMesh.h"
 
+/* TODO replace this with rendering property enum when extending the system */
+UENUM()
+enum class EFlattenMaterialProperties : uint8
+{
+	Diffuse,
+	Normal,
+	Metallic,
+	Roughness,
+	Specular,
+	Opacity,
+	Emissive,
+	SubSurface,
+	NumFlattenMaterialProperties
+};
+
 /** Structure used for storing intermediate baked down material data/samples*/
 struct FFlattenMaterial
 {
 	FFlattenMaterial()
-		: RenderSize( 0, 0 )
-		, DiffuseSize(0, 0)
-		, NormalSize(0, 0)
-		, MetallicSize(0, 0)
-		, RoughnessSize(0, 0)
-		, SpecularSize(0, 0)
-		, OpacitySize(0, 0)
-		, EmissiveSize(0, 0)
-		, SubSurfaceSize(0, 0)
+		: RenderSize(0, 0)
 		, bTwoSided(false)
 		, bDitheredLODTransition(false)
 		, BlendMode(BLEND_Opaque)
-		, EmissiveScale(1.0f)		
-	{}
+		, EmissiveScale(1.0f)
+	{
+		for (FIntPoint& Size : PropertySizes)
+		{
+			Size = FIntPoint(ForceInitToZero);
+		}
+	}
 
 	/** Release all the sample data */
 	void ReleaseData()
 	{
-		DiffuseSamples.Empty();
-		NormalSamples.Empty();
-		MetallicSamples.Empty();
-		RoughnessSamples.Empty();
-		SpecularSamples.Empty();
-		OpacitySamples.Empty();
-		EmissiveSamples.Empty();
-		SubSurfaceSamples.Empty();
+		for (TArray<FColor>& Samples : PropertySamples)
+		{
+			Samples.Empty();
+		}
 	}
 	
 	/** Set all alpha channel values with InAlphaValue */
 	void FillAlphaValues(const uint8 InAlphaValue)
 	{
-		for (FColor& Sample : DiffuseSamples) 
-		{ 
-			Sample.A = InAlphaValue;
-		}
-		for (FColor& Sample : NormalSamples) 
-		{ 
-			Sample.A = InAlphaValue;
-		}
-		for (FColor& Sample : MetallicSamples) 
-		{ 
-			Sample.A = InAlphaValue;
-		}
-		for (FColor& Sample : RoughnessSamples) 
-		{ 
-			Sample.A = InAlphaValue;
-		}
-		for (FColor& Sample : SpecularSamples)
-		{ 
-			Sample.A = InAlphaValue;
-		}
-		for (FColor& Sample : OpacitySamples) 
-		{ 
-			Sample.A = InAlphaValue;
-		}
-		for (FColor& Sample : EmissiveSamples)
-		{ 
-			Sample.A = InAlphaValue; 
-		}
-		for (FColor& Sample : SubSurfaceSamples)
+		for (TArray<FColor>& SampleArray : PropertySamples)
 		{
-			Sample.A = InAlphaValue;
+			for (FColor& Sample : SampleArray)
+			{
+				Sample.A = InAlphaValue;
+			}
 		}
 	}
+
+	const bool DoesPropertyContainData(const EFlattenMaterialProperties Property) const { return PropertySamples[(int32)Property].Num() > 0; }
+
+	const bool IsPropertyConstant(const EFlattenMaterialProperties Property) const { return PropertySamples[(int32)Property].Num() == 1; }
+
+	const bool ShouldGenerateDataForProperty(const EFlattenMaterialProperties Property) const { return PropertySizes[(int32)Property].GetMin() > 0; }
+
+	const FIntPoint GetPropertySize(const EFlattenMaterialProperties Property) const{ return PropertySizes[(int32)Property]; }
+	void SetPropertySize(const EFlattenMaterialProperties Property, const FIntPoint& InSize) { PropertySizes[(int32)Property] = InSize; }
+
+	TArray<FColor>& GetPropertySamples(const EFlattenMaterialProperties Property) { return PropertySamples[(int32)Property]; }
+	const TArray<FColor>& GetPropertySamples(const EFlattenMaterialProperties Property) const { return PropertySamples[(int32)Property]; }
 	
 	/** Material Guid */
-	FGuid			MaterialId;
-		
-	/** Texture sizes for each individual property*/
+	FGuid			MaterialId;	
 	FIntPoint		RenderSize;
-	FIntPoint		DiffuseSize;
-	FIntPoint		NormalSize;
-	FIntPoint		MetallicSize;	
-	FIntPoint		RoughnessSize;	
-	FIntPoint		SpecularSize;	
-	FIntPoint		OpacitySize;
-	FIntPoint		EmissiveSize;
-	FIntPoint		SubSurfaceSize;
-
-	/** Baked down texture samples for each individual property*/
-	TArray<FColor>	DiffuseSamples;
-	TArray<FColor>	NormalSamples;
-	TArray<FColor>	MetallicSamples;
-	TArray<FColor>	RoughnessSamples;
-	TArray<FColor>	SpecularSamples;
-	TArray<FColor>	OpacitySamples;
-	TArray<FColor>	EmissiveSamples;
-	TArray<FColor>	SubSurfaceSamples;
 
 	/** Flag whether or not the material will have to be two-sided */
 	bool			bTwoSided;
@@ -112,6 +87,12 @@ struct FFlattenMaterial
 	EBlendMode		BlendMode;
 	/** Scale (maximum baked down value) for the emissive property */
 	float			EmissiveScale;
+private:
+
+	/** Texture sizes for each individual property*/
+	FIntPoint PropertySizes[(uint32)EFlattenMaterialProperties::NumFlattenMaterialProperties];
+	/** Baked down texture samples for each individual property*/
+	TArray<FColor> PropertySamples[(uint32)EFlattenMaterialProperties::NumFlattenMaterialProperties];
 };
 
 /** Export material proxy cache*/
@@ -455,6 +436,9 @@ public:
 	*/
 	static bool ExportMaterialTexCoordScales(UMaterialInterface* InMaterial, EMaterialQualityLevel::Type QualityLevel, ERHIFeatureLevel::Type FeatureLevel, TArray<FMaterialTexCoordBuildInfo>& OutScales, FExportErrorManager& OutErrors);
 
+	// QQQ COMMENTS
+	static void DetermineMaterialImportance(const TArray<UMaterialInterface*>& InMaterials, TArray<float>& OutImportance);
+	static void GeneratedBinnedTextureSquares(const FVector2D DestinationSize, TArray<float>& InTexureWeights, TArray<FBox2D>& OutGeneratedBoxes);
 private:
 	
 	/**
@@ -478,7 +462,7 @@ private:
 	* @param OutSamples				Array of FColor samples containing the rendered out texture pixel data
 	* @return						Whether operation was successful
 	*/
-	static bool RenderMaterialPropertyToTexture(struct FMaterialMergeData& InMaterialData, EMaterialProperty InMaterialProperty, bool bInForceLinearGamma, EPixelFormat InPixelFormat, const FIntPoint& InTargetSize, FIntPoint& OutSampleSize, TArray<FColor>& OutSamples);
+	static bool RenderMaterialPropertyToTexture(struct FMaterialMergeData& InMaterialData, EMaterialProperty InMaterialProperty, bool bInForceLinearGamma, EPixelFormat InPixelFormat, const FIntPoint InTargetSize, FIntPoint& OutSampleSize, TArray<FColor>& OutSamples);
 
 	/**
 	* Creates and add or reuses a RenderTarget from the pool
@@ -503,6 +487,7 @@ private:
 
 	/** Tries to optimize the sample array (will set to const value if all samples are equal) */
 	static void OptimizeSampleArray(TArray<FColor>& InSamples, FIntPoint& InSampleSize);
+	
 private:
 	/** Flag to indicate whether or not a texture is currently being rendered out */
 	static bool CurrentlyRendering;
diff --git a/Engine/Source/Developer/MeshUtilities/MeshUtilities.Build.cs b/Engine/Source/Developer/MeshUtilities/MeshUtilities.Build.cs
--- a/Engine/Source/Developer/MeshUtilities/MeshUtilities.Build.cs
+++ b/Engine/Source/Developer/MeshUtilities/MeshUtilities.Build.cs
@@ -22,13 +22,19 @@ public class MeshUtilities : ModuleRules
 				"RawMesh",
 				"RenderCore", // For FPackedNormal
 				"SlateCore",
-				"MaterialUtilities", 
-                "MeshBoneReduction",		
+                "Slate",
+                "MaterialUtilities",
+                "MeshBoneReduction",
                 "UnrealEd",
-                "RHI",                
+                "RHI",
                 "HierarchicalLODUtilities",
-                "Landscape"
-			}
+                "Landscape",
+                "LevelEditor",
+                "AnimationBlueprintEditor",
+                "AnimationEditor",
+                "SkeletalMeshEditor",
+                "SkeletonEditor",
+            }
 		);
 
         AddEngineThirdPartyPrivateStaticDependencies(Target, "nvTriStrip");
