Fix -monolithic and -modular command line arguments not being parsed by UBT (causing monolithic-editor CIS steps to just build in modular configuration).

Also workaround several issues preventing monolithic editor compile from succeeding.

diff --git a/Engine/Plugins/Developer/VisualStudioCodeSourceCodeAccess/Source/VisualStudioCodeSourceCodeAccess/Private/VisualStudioCodeSourceCodeAccessor.cpp b/Engine/Plugins/Developer/VisualStudioCodeSourceCodeAccess/Source/VisualStudioCodeSourceCodeAccess/Private/VisualStudioCodeSourceCodeAccessor.cpp
--- a/Engine/Plugins/Developer/VisualStudioCodeSourceCodeAccess/Source/VisualStudioCodeSourceCodeAccess/Private/VisualStudioCodeSourceCodeAccessor.cpp
+++ b/Engine/Plugins/Developer/VisualStudioCodeSourceCodeAccess/Source/VisualStudioCodeSourceCodeAccess/Private/VisualStudioCodeSourceCodeAccessor.cpp
@@ -41,7 +41,7 @@ FString FVisualStudioCodeSourceCodeAccessor::GetSolutionPath() const
 }
 
 /** save all open documents in visual studio, when recompiling */
-void OnModuleCompileStarted(bool bIsAsyncCompile)
+static void OnModuleCompileStarted(bool bIsAsyncCompile)
 {
 	FVisualStudioCodeSourceCodeAccessModule& VisualStudioCodeSourceCodeAccessModule = FModuleManager::LoadModuleChecked<FVisualStudioCodeSourceCodeAccessModule>(TEXT("VisualStudioCodeSourceCodeAccess"));
 	VisualStudioCodeSourceCodeAccessModule.GetAccessor().SaveAllOpenDocuments();
diff --git a/Engine/Plugins/Developer/VisualStudioSourceCodeAccess/Source/VisualStudioSourceCodeAccess/Private/VisualStudioSourceCodeAccessor.cpp b/Engine/Plugins/Developer/VisualStudioSourceCodeAccess/Source/VisualStudioSourceCodeAccess/Private/VisualStudioSourceCodeAccessor.cpp
--- a/Engine/Plugins/Developer/VisualStudioSourceCodeAccess/Source/VisualStudioSourceCodeAccess/Private/VisualStudioSourceCodeAccessor.cpp
+++ b/Engine/Plugins/Developer/VisualStudioSourceCodeAccess/Source/VisualStudioSourceCodeAccess/Private/VisualStudioSourceCodeAccessor.cpp
@@ -85,7 +85,7 @@ enum class EAccessVisualStudioResult : uint8
 };
 
 /** save all open documents in visual studio, when recompiling */
-void OnModuleCompileStarted(bool bIsAsyncCompile)
+static void OnModuleCompileStarted(bool bIsAsyncCompile)
 {
 	FVisualStudioSourceCodeAccessModule& VisualStudioSourceCodeAccessModule = FModuleManager::LoadModuleChecked<FVisualStudioSourceCodeAccessModule>(TEXT("VisualStudioSourceCodeAccess"));
 	VisualStudioSourceCodeAccessModule.GetAccessor().SaveAllOpenDocuments();
diff --git a/Engine/Plugins/Experimental/ControlRig/Source/ControlRig/Private/Rigs/ControlManipulator.cpp b/Engine/Plugins/Experimental/ControlRig/Source/ControlRig/Private/Rigs/ControlManipulator.cpp
--- a/Engine/Plugins/Experimental/ControlRig/Source/ControlRig/Private/Rigs/ControlManipulator.cpp
+++ b/Engine/Plugins/Experimental/ControlRig/Source/ControlRig/Private/Rigs/ControlManipulator.cpp
@@ -33,7 +33,7 @@ struct FPropertyAndIndex
 	int32 ArrayIndex;
 };
 
-FPropertyAndIndex FindPropertyAndArrayIndex(UStruct* InStruct, const FString& PropertyName)
+static FPropertyAndIndex FindPropertyAndArrayIndex2(UStruct* InStruct, const FString& PropertyName)
 {
 	FPropertyAndIndex PropertyAndIndex;
 
@@ -71,7 +71,7 @@ FPropertyAddress FindPropertyRecursive(void* BasePointer, UStruct* InStruct, TAr
 	check(InOutPropertyPath);
 #endif
 
-	FPropertyAndIndex PropertyAndIndex = FindPropertyAndArrayIndex(InStruct, *InPropertyNames[Index]);
+	FPropertyAndIndex PropertyAndIndex = FindPropertyAndArrayIndex2(InStruct, *InPropertyNames[Index]);
 
 	FPropertyAddress NewAddress;
 
diff --git a/Engine/Plugins/Runtime/ApexDestruction/Source/ApexDestructionEditor/Private/ApexDestructionEditorModule.cpp b/Engine/Plugins/Runtime/ApexDestruction/Source/ApexDestructionEditor/Private/ApexDestructionEditorModule.cpp
--- a/Engine/Plugins/Runtime/ApexDestruction/Source/ApexDestructionEditor/Private/ApexDestructionEditorModule.cpp
+++ b/Engine/Plugins/Runtime/ApexDestruction/Source/ApexDestructionEditor/Private/ApexDestructionEditorModule.cpp
@@ -28,7 +28,7 @@
 #include "PropertyEditorModule.h"
 #include "DestructibleMeshDetails.h"
 
-IMPLEMENT_MODULE( FDestructibleMeshEditorModule, DestructibleMeshEditor );
+IMPLEMENT_MODULE( FDestructibleMeshEditorModule, ApexDestructionEditor );
 
 #define LOCTEXT_NAMESPACE "DestructibleMeshEditor"
 
diff --git a/Engine/Source/Developer/AudioSettingsEditor/Private/TargetPlatformAudioCustomization.cpp b/Engine/Source/Developer/AudioSettingsEditor/Private/TargetPlatformAudioCustomization.cpp
--- a/Engine/Source/Developer/AudioSettingsEditor/Private/TargetPlatformAudioCustomization.cpp
+++ b/Engine/Source/Developer/AudioSettingsEditor/Private/TargetPlatformAudioCustomization.cpp
@@ -16,6 +16,8 @@
 #include "AudioDevice.h"
 #endif 
 
+IMPLEMENT_MODULE(FDefaultModuleImpl, AudioSettingsEditor)
+
 #define LOCTEXT_NAMESPACE "PlatformAudio"
 
 FAudioPluginWidgetManager::FAudioPluginWidgetManager()
diff --git a/Engine/Source/Developer/MaterialUtilities/Private/MeshRendering.cpp b/Engine/Source/Developer/MaterialUtilities/Private/MeshRendering.cpp
--- a/Engine/Source/Developer/MaterialUtilities/Private/MeshRendering.cpp
+++ b/Engine/Source/Developer/MaterialUtilities/Private/MeshRendering.cpp
@@ -183,10 +183,10 @@ public:
 /**
 * Canvas render item enqueued into renderer command list.
 */
-class FMeshMaterialRenderItem : public FCanvasBaseRenderItem
+class FMeshMaterialRenderItem2 : public FCanvasBaseRenderItem
 {
 public:
-	FMeshMaterialRenderItem(FSceneViewFamily* InViewFamily, const FRawMesh* InMesh, const FSkeletalMeshLODRenderData* InLODData, int32 LightMapIndex, int32 InMaterialIndex, const FBox2D& InTexcoordBounds, const TArray<FVector2D>& InTexCoords, const FVector2D& InSize, const FMaterialRenderProxy* InMaterialRenderProxy, const FCanvas::FTransformEntry& InTransform /*= FCanvas::FTransformEntry(FMatrix::Identity)*/, FLightMapRef LightMap, FShadowMapRef ShadowMap, FUniformBufferRHIRef Buffer) : Data(new FRenderData(
+	FMeshMaterialRenderItem2(FSceneViewFamily* InViewFamily, const FRawMesh* InMesh, const FSkeletalMeshLODRenderData* InLODData, int32 LightMapIndex, int32 InMaterialIndex, const FBox2D& InTexcoordBounds, const TArray<FVector2D>& InTexCoords, const FVector2D& InSize, const FMaterialRenderProxy* InMaterialRenderProxy, const FCanvas::FTransformEntry& InTransform /*= FCanvas::FTransformEntry(FMatrix::Identity)*/, FLightMapRef LightMap, FShadowMapRef ShadowMap, FUniformBufferRHIRef Buffer) : Data(new FRenderData(
 		InViewFamily,
 		InMesh,
 		InLODData,
@@ -201,7 +201,7 @@ public:
 	{
 	}
 
-	~FMeshMaterialRenderItem()
+	~FMeshMaterialRenderItem2()
 	{
 	}
 
@@ -257,7 +257,7 @@ public:
 		// get the current transform entry from top of transform stack
 		const FCanvas::FTransformEntry& TopTransformEntry = InCanvas->GetTransformStack().Top();
 		// create a render batch
-		FMeshMaterialRenderItem* RenderBatch = new FMeshMaterialRenderItem(
+		FMeshMaterialRenderItem2* RenderBatch = new FMeshMaterialRenderItem2(
 			InViewFamily,
 			InMesh,
 			InLODRenderData,
@@ -780,7 +780,7 @@ bool FMeshRenderer::RenderMaterial(struct FMaterialMergeData& InMaterialData, FM
 #endif
 
 		// add item for rendering
-		FMeshMaterialRenderItem::EnqueueMaterialRender(
+		FMeshMaterialRenderItem2::EnqueueMaterialRender(
 			&Canvas,
 			&ViewFamily,
 			InMaterialData.Mesh,
@@ -909,7 +909,7 @@ bool FMeshRenderer::RenderMaterialTexCoordScales(struct FMaterialMergeData& InMa
 	ViewFamily.bNullifyWorldSpacePosition = true;
 
 	// add item for rendering
-	FMeshMaterialRenderItem::EnqueueMaterialRender(
+	FMeshMaterialRenderItem2::EnqueueMaterialRender(
 		&Canvas,
 		&ViewFamily,
 		InMaterialData.Mesh,
diff --git a/Engine/Source/Developer/XGEController/Private/XGEControllerModule.cpp b/Engine/Source/Developer/XGEController/Private/XGEControllerModule.cpp
--- a/Engine/Source/Developer/XGEController/Private/XGEControllerModule.cpp
+++ b/Engine/Source/Developer/XGEController/Private/XGEControllerModule.cpp
@@ -573,6 +573,6 @@ XGECONTROLLER_API IXGEController& IXGEController::Get()
 	return Ref;
 }
 
-IMPLEMENT_MODULE(FXGEControllerModule, XGEControllerModule);
+IMPLEMENT_MODULE(FXGEControllerModule, XGEController);
 
 #endif // WITH_XGE_CONTROLLER
diff --git a/Engine/Source/Editor/UnrealEd/Private/EditorActorFolders.cpp b/Engine/Source/Editor/UnrealEd/Private/EditorActorFolders.cpp
--- a/Engine/Source/Editor/UnrealEd/Private/EditorActorFolders.cpp
+++ b/Engine/Source/Editor/UnrealEd/Private/EditorActorFolders.cpp
@@ -28,7 +28,7 @@ FString GetWorldStateFilename(UPackage* Package)
 }
 
 /** Convert an old path to a new path, replacing an ancestor branch with something else */
-FName OldPathToNewPath(const FString& InOldBranch, const FString& InNewBranch, const FString& PathToMove)
+static FName OldPathToNewPath(const FString& InOldBranch, const FString& InNewBranch, const FString& PathToMove)
 {
 	return FName(*(InNewBranch + PathToMove.RightChop(InOldBranch.Len())));
 }
diff --git a/Engine/Source/Editor/WorldBrowser/Private/LevelFolders.cpp b/Engine/Source/Editor/WorldBrowser/Private/LevelFolders.cpp
--- a/Engine/Source/Editor/WorldBrowser/Private/LevelFolders.cpp
+++ b/Engine/Source/Editor/WorldBrowser/Private/LevelFolders.cpp
@@ -23,7 +23,7 @@ FString GetLevelModelFilename(TSharedPtr<FLevelModel> LevelModel)
 	return FPaths::Combine(*FPaths::ProjectSavedDir(), TEXT("Config"), TEXT("LevelState"), *FString::Printf(TEXT("%u.json"), PackageNameCrc));
 }
 
-FName OldPathToNewPath(const FString& InOldBranch, const FString& InNewBranch, const FString& PathToMove)
+static FName OldPathToNewPath(const FString& InOldBranch, const FString& InNewBranch, const FString& PathToMove)
 {
 	return FName(*(InNewBranch + PathToMove.RightChop(InOldBranch.Len())));
 }
diff --git a/Engine/Source/Programs/UnrealBuildTool/System/CommandLine.cs b/Engine/Source/Programs/UnrealBuildTool/System/CommandLine.cs
--- a/Engine/Source/Programs/UnrealBuildTool/System/CommandLine.cs
+++ b/Engine/Source/Programs/UnrealBuildTool/System/CommandLine.cs
@@ -53,31 +53,34 @@ namespace UnrealBuildTool
 		{
 			// Build a mapping from name to field and attribute for this object
 			Dictionary<string, Parameter> PrefixToParameter = new Dictionary<string, Parameter>(StringComparer.InvariantCultureIgnoreCase);
-			foreach(FieldInfo FieldInfo in TargetObject.GetType().GetFields(BindingFlags.Instance | BindingFlags.GetField | BindingFlags.GetProperty | BindingFlags.Public | BindingFlags.NonPublic))
+			for(Type TargetType = TargetObject.GetType(); TargetType != typeof(object); TargetType = TargetType.BaseType)
 			{
-				IEnumerable<CommandLineAttribute> Attributes = FieldInfo.GetCustomAttributes<CommandLineAttribute>();
-				foreach(CommandLineAttribute Attribute in Attributes)
+				foreach(FieldInfo FieldInfo in TargetType.GetFields(BindingFlags.Instance | BindingFlags.GetField | BindingFlags.GetProperty | BindingFlags.Public | BindingFlags.NonPublic | BindingFlags.DeclaredOnly))
 				{
-					string Prefix = Attribute.Prefix;
-					if(Prefix == null)
+					IEnumerable<CommandLineAttribute> Attributes = FieldInfo.GetCustomAttributes<CommandLineAttribute>();
+					foreach(CommandLineAttribute Attribute in Attributes)
 					{
-						if(FieldInfo.FieldType == typeof(bool))
+						string Prefix = Attribute.Prefix;
+						if(Prefix == null)
 						{
-							Prefix = String.Format("-{0}", FieldInfo.Name);
+							if(FieldInfo.FieldType == typeof(bool))
+							{
+								Prefix = String.Format("-{0}", FieldInfo.Name);
+							}
+							else
+							{
+								Prefix = String.Format("-{0}=", FieldInfo.Name);
+							}
 						}
 						else
 						{
-							Prefix = String.Format("-{0}=", FieldInfo.Name);
-						}
-					}
-					else
-					{
-						if(FieldInfo.FieldType != typeof(bool) && !Attribute.ValueAfterSpace && Attribute.Value == null && !Prefix.EndsWith("="))
-						{
-							Prefix = Prefix + "=";
+							if(FieldInfo.FieldType != typeof(bool) && !Attribute.ValueAfterSpace && Attribute.Value == null && !Prefix.EndsWith("="))
+							{
+								Prefix = Prefix + "=";
+							}
 						}
+						PrefixToParameter.Add(Prefix, new Parameter { Prefix = Prefix, Attribute = Attribute, FieldInfo = FieldInfo });
 					}
-					PrefixToParameter.Add(Prefix, new Parameter { Prefix = Prefix, Attribute = Attribute, FieldInfo = FieldInfo });
 				}
 			}
 
