PR #4305: Light map index was unintialized (Contributed by DSDambuster)

diff --git a/Engine/Source/Developer/MaterialUtilities/Public/MaterialUtilities.h b/Engine/Source/Developer/MaterialUtilities/Public/MaterialUtilities.h
--- a/Engine/Source/Developer/MaterialUtilities/Public/MaterialUtilities.h
+++ b/Engine/Source/Developer/MaterialUtilities/Public/MaterialUtilities.h
@@ -173,6 +173,7 @@ struct FMaterialMergeData
 		, TexcoordBounds(InTexcoordBounds)
 		, TexCoords(InTexCoords)
 		, EmissiveScale(0.0f)
+		, LightMapIndex(1)
 	{
 		ProxyCache = new FExportMaterialProxyCache();
 	}
