Fix CIS build, missing a forward declaration

diff --git a/Engine/Source/Developer/MaterialUtilities/Public/MaterialUtilities.h b/Engine/Source/Developer/MaterialUtilities/Public/MaterialUtilities.h
--- a/Engine/Source/Developer/MaterialUtilities/Public/MaterialUtilities.h
+++ b/Engine/Source/Developer/MaterialUtilities/Public/MaterialUtilities.h
@@ -25,6 +25,7 @@ class UTextureRenderTarget2D;
 class UMaterialOptions;
 struct FMaterialProxySettings;
 struct FMeshDescription;
+struct FRawMesh;
 class FSkeletalMeshLODRenderData;
 struct FBakeOutput;
 struct FMeshData;
