Deprecate FRawMesh - Remove all FRawMesh dependencies from material utilities

diff --git a/Engine/Source/Developer/MaterialUtilities/MaterialUtilities.Build.cs b/Engine/Source/Developer/MaterialUtilities/MaterialUtilities.Build.cs
--- a/Engine/Source/Developer/MaterialUtilities/MaterialUtilities.Build.cs
+++ b/Engine/Source/Developer/MaterialUtilities/MaterialUtilities.Build.cs
@@ -30,7 +30,6 @@ public class MaterialUtilities : ModuleRules
 
         PublicDependencyModuleNames.AddRange(
 			new string [] {
-				 "RawMesh",
                  "MeshDescription",
                  "MeshDescriptionOperations",
 			}
diff --git a/Engine/Source/Developer/MaterialUtilities/Private/MaterialUtilities.cpp b/Engine/Source/Developer/MaterialUtilities/Private/MaterialUtilities.cpp
--- a/Engine/Source/Developer/MaterialUtilities/Private/MaterialUtilities.cpp
+++ b/Engine/Source/Developer/MaterialUtilities/Private/MaterialUtilities.cpp
@@ -913,42 +913,6 @@ bool FMaterialUtilities::ExportMaterial(UMaterialInterface* InMaterial, FFlatten
 	return true;
 }
 
-bool FMaterialUtilities::ExportMaterial(UMaterialInterface* InMaterial, const FRawMesh* InMesh, int32 InMaterialIndex, const FBox2D& InTexcoordBounds, const TArray<FVector2D>& InTexCoords, FFlattenMaterial& OutFlattenMaterial, struct FExportMaterialProxyCache* ProxyCache)
-{
-	FMeshDescription MeshDescription;
-	UStaticMesh::RegisterMeshAttributes(MeshDescription);
-	TMap<int32, FName> MaterialMap;
-	FMeshDescriptionOperations::ConvertFromRawMesh(*InMesh, MeshDescription, MaterialMap);
-
-	FMaterialMergeData MaterialData(InMaterial, &MeshDescription, nullptr, InMaterialIndex, InTexcoordBounds, InTexCoords);
-
-	TArray<FMaterialMergeData*> MergeData{ &MaterialData };
-	TArray<FFlattenMaterial*> FlatMaterials{ &OutFlattenMaterial };
-	PRAGMA_DISABLE_DEPRECATION_WARNINGS
-	ExportMaterials(MergeData, FlatMaterials);
-	PRAGMA_ENABLE_DEPRECATION_WARNINGS
-	return true;
-}
-
-bool FMaterialUtilities::ExportMaterial(UMaterialInterface* InMaterial, const FRawMesh* InMesh, int32 InMaterialIndex, const FBox2D& InTexcoordBounds, const TArray<FVector2D>& InTexCoords, const int32 LightMapIndex, FLightMapRef LightMap, FShadowMapRef ShadowMap, FUniformBufferRHIRef Buffer, FFlattenMaterial& OutFlattenMaterial, struct FExportMaterialProxyCache* ProxyCache /*= NULL*/)
-{
-	FMeshDescription MeshDescription;
-	UStaticMesh::RegisterMeshAttributes(MeshDescription);
-	TMap<int32, FName> MaterialMap;
-	FMeshDescriptionOperations::ConvertFromRawMesh(*InMesh, MeshDescription, MaterialMap);
-
-	FMaterialMergeData MaterialData(InMaterial, &MeshDescription, nullptr, InMaterialIndex, InTexcoordBounds, InTexCoords);
-	MaterialData.LightMapIndex = 1;
-	MaterialData.LightMap = LightMap;
-	MaterialData.ShadowMap = ShadowMap;
-	MaterialData.Buffer = Buffer;
-	PRAGMA_DISABLE_DEPRECATION_WARNINGS
-	ExportMaterial(MaterialData, OutFlattenMaterial, ProxyCache);
-	PRAGMA_ENABLE_DEPRECATION_WARNINGS
-
-	return true;
-}
-
 bool FMaterialUtilities::ExportLandscapeMaterial(ALandscapeProxy* InLandscape, const TSet<FPrimitiveComponentId>& HiddenPrimitives, FFlattenMaterial& OutFlattenMaterial)
 {
 	check(InLandscape);
diff --git a/Engine/Source/Developer/MaterialUtilities/Public/MaterialUtilities.h b/Engine/Source/Developer/MaterialUtilities/Public/MaterialUtilities.h
--- a/Engine/Source/Developer/MaterialUtilities/Public/MaterialUtilities.h
+++ b/Engine/Source/Developer/MaterialUtilities/Public/MaterialUtilities.h
@@ -25,7 +25,6 @@ class UTextureRenderTarget2D;
 class UMaterialOptions;
 struct FMaterialProxySettings;
 struct FMeshDescription;
-struct FRawMesh;
 class FSkeletalMeshLODRenderData;
 struct FBakeOutput;
 struct FMeshData;
@@ -289,23 +288,6 @@ public:
 	DEPRECATED(4.17, "Please use new functionality in MaterialBaking module")
 	static bool ExportMaterials(TArray<FMaterialMergeData*>& MergeData, TArray<FFlattenMaterial*>& OutFlattenMaterials);
 
-	/**
-	* Flattens specified material using mesh data	
-	*
-	* @param InMaterial			Target material
-	* @param InMesh				Mesh data to use for flattening
-	* @param InMaterialIndex	Material index
-	* @param InTexcoordBounds	Texture bounds 
-	* @param InTexCoords		Replacement non-overlapping texture coordinates
-	* @param OutFlattenMaterial Output flattened material
-	* @return					Whether operation was successful
-	*/
-	DEPRECATED(4.19, "Please use new functionality in MaterialBaking module")
-	static bool ExportMaterial(UMaterialInterface* InMaterial, const FRawMesh* InMesh, int32 InMaterialIndex, const FBox2D& InTexcoordBounds, const TArray<FVector2D>& InTexCoords, FFlattenMaterial& OutFlattenMaterial, struct FExportMaterialProxyCache* ProxyCache = nullptr);
-
-	DEPRECATED(4.19, "Please use new functionality in MaterialBaking module")
-	static bool ExportMaterial(UMaterialInterface* InMaterial, const FRawMesh* InMesh, int32 InMaterialIndex, const FBox2D& InTexcoordBounds, const TArray<FVector2D>& InTexCoords, const int32 LightMapIndex, FLightMapRef LightMap, FShadowMapRef ShadowMap, FUniformBufferRHIRef Buffer, FFlattenMaterial& OutFlattenMaterial, struct FExportMaterialProxyCache* ProxyCache = nullptr);
-
 	/**
 	 * Flattens specified landscape material
 	 *
