Remove SCOPED_SUSPEND_RENDERING_THREAD from OBJs exporter

Alexis.Matte

diff --git a/Engine/Source/Developer/MaterialBaking/Private/ExportMaterialProxy.h b/Engine/Source/Developer/MaterialBaking/Private/ExportMaterialProxy.h
--- a/Engine/Source/Developer/MaterialBaking/Private/ExportMaterialProxy.h
+++ b/Engine/Source/Developer/MaterialBaking/Private/ExportMaterialProxy.h
@@ -540,93 +540,6 @@ public:
 		return Ar << V.MaterialInterface;
 	}
 
-	/**
-	* Iterate through all textures used by the material and return the maximum texture resolution used
-	* (ideally this could be made dependent of the material property)
-	*
-	* @param MaterialInterface The material to scan for texture size
-	*
-	* @return Size (width and height)
-	*/
-	FIntPoint FindMaxTextureSize(UMaterialInterface* InMaterialInterface, FIntPoint MinimumSize = FIntPoint(1, 1)) const
-	{
-		// static lod settings so that we only initialize them once
-		const UTextureLODSettings* GameTextureLODSettings = UDeviceProfileManager::Get().GetActiveProfile()->GetTextureLODSettings();
-
-		TArray<UTexture*> MaterialTextures;
-		InMaterialInterface->GetUsedTextures(MaterialTextures, EMaterialQualityLevel::Num, false, GMaxRHIFeatureLevel, false);
-
-		FTextureCompilingManager::Get().FinishCompilation(MaterialTextures);
-
-		// find the largest texture in the list (applying it's LOD bias)
-		FIntPoint MaxSize = MinimumSize;
-		for (int32 TexIndex = 0; TexIndex < MaterialTextures.Num(); TexIndex++)
-		{
-			const UTexture* Texture = MaterialTextures[TexIndex];
-
-			if (Texture == NULL)
-			{
-				continue;
-			}
-
-			// get the max size of the texture
-			const FIntPoint LocalSize = [&]()
-			{
-				if (Texture->IsA(UTexture2D::StaticClass()))
-				{
-					UTexture2D* Tex2D = (UTexture2D*)Texture;
-					return FIntPoint(Tex2D->GetSizeX(), Tex2D->GetSizeY());
-				}
-				else if (Texture->IsA(UTextureCube::StaticClass()))
-				{
-					UTextureCube* TexCube = (UTextureCube*)Texture;
-					return FIntPoint(TexCube->GetSizeX(), TexCube->GetSizeY());
-				}
-				else if (Texture->IsA(UTexture2DArray::StaticClass())) 
-				{
-					UTexture2DArray* TexArray = (UTexture2DArray*)Texture;
-					return FIntPoint(TexArray->GetSizeX(), TexArray->GetSizeY());
-				}
-				return FIntPoint(0, 0);
-			}();
-			
-			// bias the texture size based on LOD group
-			const int32 LocalBias = GameTextureLODSettings->CalculateLODBias(Texture);
-			MaxSize.X = FMath::Max(LocalSize.X >> LocalBias, MaxSize.X);
-			MaxSize.Y = FMath::Max(LocalSize.Y >> LocalBias, MaxSize.Y);
-		}
-
-		return MaxSize;
-	}
-
-	static bool WillFillData(EBlendMode InBlendMode, EMaterialProperty InMaterialProperty)
-	{
-		if (InMaterialProperty == MP_EmissiveColor)
-		{
-			return true;
-		}
-
-		switch (InBlendMode)
-		{
-		case BLEND_Opaque:
-		{
-			switch (InMaterialProperty)
-			{
-			case MP_BaseColor:				return true;
-			case MP_Specular:				return true;
-			case MP_Normal:					return true;
-			case MP_Tangent:				return true;
-			case MP_Metallic:				return true;
-			case MP_Roughness:				return true;
-			case MP_Anisotropy:				return true;
-			case MP_AmbientOcclusion:		return true;
-			}
-		}
-		break;
-		}
-		return false;
-	}
-
 	virtual bool IsUsedWithStaticLighting() const override
 	{
 		return true; 
diff --git a/Engine/Source/Developer/MaterialUtilities/Private/MaterialUtilities.cpp b/Engine/Source/Developer/MaterialUtilities/Private/MaterialUtilities.cpp
--- a/Engine/Source/Developer/MaterialUtilities/Private/MaterialUtilities.cpp
+++ b/Engine/Source/Developer/MaterialUtilities/Private/MaterialUtilities.cpp
@@ -695,65 +695,6 @@ public:
 		return Ar << V.MaterialInterface;
 	}
 
-	/**
-	* Iterate through all textures used by the material and return the maximum texture resolution used
-	* (ideally this could be made dependent of the material property)
-	*
-	* @param MaterialInterface The material to scan for texture size
-	*
-	* @return Size (width and height)
-	*/
-	FIntPoint FindMaxTextureSize(UMaterialInterface* InMaterialInterface, FIntPoint MinimumSize = FIntPoint(1, 1)) const
-	{
-		// static lod settings so that we only initialize them once
-		UTextureLODSettings* GameTextureLODSettings = UDeviceProfileManager::Get().GetActiveProfile()->GetTextureLODSettings();
-
-		TArray<UTexture*> MaterialTextures;
-
-		InMaterialInterface->GetUsedTextures(MaterialTextures, EMaterialQualityLevel::Num, false, GMaxRHIFeatureLevel, false);
-		FTextureCompilingManager::Get().FinishCompilation(MaterialTextures);
-
-		// find the largest texture in the list (applying it's LOD bias)
-		FIntPoint MaxSize = MinimumSize;
-		for (int32 TexIndex = 0; TexIndex < MaterialTextures.Num(); TexIndex++)
-		{
-			UTexture* Texture = MaterialTextures[TexIndex];
-
-			if (Texture == NULL)
-			{
-				continue;
-			}
-
-			// get the max size of the texture
-			FIntPoint LocalSize(0, 0);
-			if (Texture->IsA(UTexture2D::StaticClass()))
-			{
-				UTexture2D* Tex2D = (UTexture2D*)Texture;
-				LocalSize = FIntPoint(Tex2D->GetSizeX(), Tex2D->GetSizeY());
-			}
-			else if (Texture->IsA(UTextureCube::StaticClass()))
-			{
-				UTextureCube* TexCube = (UTextureCube*)Texture;
-				LocalSize = FIntPoint(TexCube->GetSizeX(), TexCube->GetSizeY());
-			}
-			else if (Texture->IsA(UTexture2DArray::StaticClass())) 
-			{
-				UTexture2DArray* TexArray = (UTexture2DArray*)Texture;
-				LocalSize = FIntPoint(TexArray->GetSizeX(), TexArray->GetSizeY());
-			}
-
-			int32 LocalBias = GameTextureLODSettings->CalculateLODBias(Texture);
-
-			// bias the texture size based on LOD group
-			FIntPoint BiasedLocalSize(LocalSize.X >> LocalBias, LocalSize.Y >> LocalBias);
-
-			MaxSize.X = FMath::Max(BiasedLocalSize.X, MaxSize.X);
-			MaxSize.Y = FMath::Max(BiasedLocalSize.Y, MaxSize.Y);
-		}
-
-		return MaxSize;
-	}
-
 	static bool WillFillData(EBlendMode InBlendMode, EMaterialProperty InMaterialProperty)
 	{
 		if (InMaterialProperty == MP_EmissiveColor)
@@ -865,7 +806,56 @@ static void RenderSceneToTexture(
 	RenderTargetTexture = nullptr;
 }
 
+FIntPoint FMaterialUtilities::FindMaxTextureSize(UMaterialInterface* InMaterialInterface, FIntPoint MinimumSize)
+{
+	// static lod settings so that we only initialize them once
+	UTextureLODSettings* GameTextureLODSettings = UDeviceProfileManager::Get().GetActiveProfile()->GetTextureLODSettings();
+
+	TArray<UTexture*> MaterialTextures;
+
+	InMaterialInterface->GetUsedTextures(MaterialTextures, EMaterialQualityLevel::Num, false, GMaxRHIFeatureLevel, false);
+	FTextureCompilingManager::Get().FinishCompilation(MaterialTextures);
+
+	// find the largest texture in the list (applying it's LOD bias)
+	FIntPoint MaxSize = MinimumSize;
+	for (int32 TexIndex = 0; TexIndex < MaterialTextures.Num(); TexIndex++)
+	{
+		UTexture* Texture = MaterialTextures[TexIndex];
+
+		if (Texture == NULL)
+		{
+			continue;
+		}
+
+		// get the max size of the texture
+		FIntPoint LocalSize(0, 0);
+		if (Texture->IsA(UTexture2D::StaticClass()))
+		{
+			UTexture2D* Tex2D = (UTexture2D*)Texture;
+			LocalSize = FIntPoint(Tex2D->GetSizeX(), Tex2D->GetSizeY());
+		}
+		else if (Texture->IsA(UTextureCube::StaticClass()))
+		{
+			UTextureCube* TexCube = (UTextureCube*)Texture;
+			LocalSize = FIntPoint(TexCube->GetSizeX(), TexCube->GetSizeY());
+		}
+		else if (Texture->IsA(UTexture2DArray::StaticClass()))
+		{
+			UTexture2DArray* TexArray = (UTexture2DArray*)Texture;
+			LocalSize = FIntPoint(TexArray->GetSizeX(), TexArray->GetSizeY());
+		}
+
+		int32 LocalBias = GameTextureLODSettings->CalculateLODBias(Texture);
 
+		// bias the texture size based on LOD group
+		FIntPoint BiasedLocalSize(LocalSize.X >> LocalBias, LocalSize.Y >> LocalBias);
+
+		MaxSize.X = FMath::Max(BiasedLocalSize.X, MaxSize.X);
+		MaxSize.Y = FMath::Max(BiasedLocalSize.Y, MaxSize.Y);
+	}
+
+	return MaxSize;
+}
 
 bool FMaterialUtilities::SupportsExport(EBlendMode InBlendMode, EMaterialProperty InMaterialProperty)
 {
@@ -874,18 +864,12 @@ bool FMaterialUtilities::SupportsExport(EBlendMode InBlendMode, EMaterialPropert
 
 bool FMaterialUtilities::ExportMaterialProperty(UWorld* InWorld, UMaterialInterface* InMaterial, EMaterialProperty InMaterialProperty, UTextureRenderTarget2D* InRenderTarget, TArray<FColor>& OutBMP)
 {
-	TUniquePtr<FExportMaterialProxy> MaterialProxy(new FExportMaterialProxy(InMaterial, InMaterialProperty));
-	if (MaterialProxy == nullptr)
-	{
-		return false;
-	}
-
 	FBox2D DummyBounds(FVector2D(0, 0), FVector2D(1, 1));
 	TArray<FVector2D> EmptyTexCoords;
 	FMaterialMergeData MaterialData(InMaterial, nullptr, nullptr, 0, DummyBounds, EmptyTexCoords);
 	const bool bForceGamma = (InMaterialProperty == MP_Normal) || (InMaterialProperty == MP_OpacityMask) || (InMaterialProperty == MP_Opacity);	
 
-	FIntPoint MaxSize = MaterialProxy->FindMaxTextureSize(InMaterial);
+	FIntPoint MaxSize = FindMaxTextureSize(InMaterial);
 	FIntPoint OutSize = MaxSize;
 	PRAGMA_DISABLE_DEPRECATION_WARNINGS
 	return RenderMaterialPropertyToTexture(MaterialData, InMaterialProperty, bForceGamma, PF_B8G8R8A8, MaxSize, OutSize, OutBMP);
@@ -894,17 +878,11 @@ bool FMaterialUtilities::ExportMaterialProperty(UWorld* InWorld, UMaterialInterf
 
 bool FMaterialUtilities::ExportMaterialProperty(UWorld* InWorld, UMaterialInterface* InMaterial, EMaterialProperty InMaterialProperty, FIntPoint& OutSize, TArray<FColor>& OutBMP)
 {
-	TUniquePtr<FExportMaterialProxy> MaterialProxy(new FExportMaterialProxy(InMaterial, InMaterialProperty));
-	if (MaterialProxy == nullptr)
-	{
-		return false;
-	}
-
 	FBox2D DummyBounds(FVector2D(0, 0), FVector2D(1, 1));
 	TArray<FVector2D> EmptyTexCoords;
 	FMaterialMergeData MaterialData(InMaterial, nullptr, nullptr, 0, DummyBounds, EmptyTexCoords);
 	const bool bForceGamma = (InMaterialProperty == MP_Normal) || (InMaterialProperty == MP_OpacityMask) || (InMaterialProperty == MP_Opacity);
-	OutSize = MaterialProxy->FindMaxTextureSize(InMaterial);
+	OutSize = FindMaxTextureSize(InMaterial);
 	PRAGMA_DISABLE_DEPRECATION_WARNINGS
 	return RenderMaterialPropertyToTexture(MaterialData, InMaterialProperty, bForceGamma, PF_B8G8R8A8, OutSize, OutSize, OutBMP);
 	PRAGMA_ENABLE_DEPRECATION_WARNINGS
@@ -912,17 +890,11 @@ bool FMaterialUtilities::ExportMaterialProperty(UWorld* InWorld, UMaterialInterf
 
 bool FMaterialUtilities::ExportMaterialProperty(UMaterialInterface* InMaterial, EMaterialProperty InMaterialProperty, TArray<FColor>& OutBMP, FIntPoint& OutSize)
 {
-	TUniquePtr<FExportMaterialProxy> MaterialProxy(new FExportMaterialProxy(InMaterial, InMaterialProperty));
-	if (MaterialProxy == nullptr)
-	{
-		return false;
-	}
-
 	FBox2D DummyBounds(FVector2D(0, 0), FVector2D(1, 1));
 	TArray<FVector2D> EmptyTexCoords;
 	FMaterialMergeData MaterialData(InMaterial, nullptr, nullptr, 0, DummyBounds, EmptyTexCoords);
 	const bool bForceGamma = (InMaterialProperty == MP_Normal) || (InMaterialProperty == MP_OpacityMask) || (InMaterialProperty == MP_Opacity);
-	OutSize = MaterialProxy->FindMaxTextureSize(InMaterial);
+	OutSize = FindMaxTextureSize(InMaterial);
 	PRAGMA_DISABLE_DEPRECATION_WARNINGS
 	return RenderMaterialPropertyToTexture(MaterialData, InMaterialProperty, bForceGamma, PF_B8G8R8A8, OutSize, OutSize, OutBMP);
 	PRAGMA_ENABLE_DEPRECATION_WARNINGS
@@ -930,12 +902,6 @@ bool FMaterialUtilities::ExportMaterialProperty(UMaterialInterface* InMaterial,
 
 bool FMaterialUtilities::ExportMaterialProperty(UMaterialInterface* InMaterial, EMaterialProperty InMaterialProperty, FIntPoint InSize, TArray<FColor>& OutBMP)
 {
-	TUniquePtr<FExportMaterialProxy> MaterialProxy(new FExportMaterialProxy(InMaterial, InMaterialProperty));
-	if (MaterialProxy == nullptr)
-	{
-		return false;
-	}
-
 	FBox2D DummyBounds(FVector2D(0, 0), FVector2D(1, 1));
 	TArray<FVector2D> EmptyTexCoords;
 	FMaterialMergeData MaterialData(InMaterial, nullptr, nullptr, 0, DummyBounds, EmptyTexCoords);
diff --git a/Engine/Source/Developer/MaterialUtilities/Public/MaterialUtilities.h b/Engine/Source/Developer/MaterialUtilities/Public/MaterialUtilities.h
--- a/Engine/Source/Developer/MaterialUtilities/Public/MaterialUtilities.h
+++ b/Engine/Source/Developer/MaterialUtilities/Public/MaterialUtilities.h
@@ -215,6 +215,16 @@ public:
 	virtual void ShutdownModule() override;
 	/** End IModuleInterface implementation */
 
+	/**
+	* Iterate through all textures used by the material and return the maximum texture resolution used
+	*
+	* @param MaterialInterface	The material to scan for texture size
+	* @param MinimumSize		The minimum size returned
+	*
+	* @return Size (width and height)
+	*/
+	static FIntPoint FindMaxTextureSize(UMaterialInterface* InMaterialInterface, FIntPoint MinimumSize = FIntPoint(1, 1));
+
 	/**
 	 * Whether material utilities support exporting specified material blend mode and property 
 	 */
diff --git a/Engine/Source/Editor/UnrealEd/Private/EditorExporters.cpp b/Engine/Source/Editor/UnrealEd/Private/EditorExporters.cpp
--- a/Engine/Source/Editor/UnrealEd/Private/EditorExporters.cpp
+++ b/Engine/Source/Editor/UnrealEd/Private/EditorExporters.cpp
@@ -72,6 +72,9 @@
 #include "InstancedFoliage.h"
 #include "Engine/Selection.h"
 #include "AssetExportTask.h"
+#include "IMaterialBakingModule.h"
+#include "MaterialBakingStructures.h"
+#include "MaterialOptions.h"
 
 #include "Components/SkeletalMeshComponent.h"
 #include "Camera/CameraComponent.h"
@@ -810,7 +813,7 @@ public:
 
 inline FString FixupMaterialName(UMaterialInterface* Material)
 {
-	return Material->GetPathName().Replace(TEXT("."), TEXT("_")).Replace(TEXT(":"), TEXT("_"));
+	return Material->GetName().Replace(TEXT("."), TEXT("_")).Replace(TEXT(":"), TEXT("_"));
 }
 
 
@@ -1013,30 +1016,65 @@ static void AddActorToOBJs(AActor* Actor, TArray<FOBJGeom*>& Objects, TSet<UMate
 
 // @param Material must not be 0
 // @param MatProp e.g. MP_DiffuseColor
-static void ExportMaterialPropertyTexture(const FString &BMPFilename, UMaterialInterface* Material, const EMaterialProperty MatProp)
+static void ExportMaterialPropertyTexture(const FString& BMPFilename, UMaterialInterface* Material, const EMaterialProperty MatProp)
 {
 	check(Material);
 
 	// make the BMP for the diffuse channel
 	TArray<FColor> OutputBMP;
 	FIntPoint OutSize;
-	
-	TEnumAsByte<EBlendMode> BlendMode = Material->GetBlendMode();
-	bool bIsValidMaterial = FMaterialUtilities::SupportsExport((EBlendMode)(BlendMode), MatProp);
 
-	if (bIsValidMaterial)
+	const TEnumAsByte<EBlendMode> BlendMode = Material->GetBlendMode();
+
+	bool bIsValidProperty = FMaterialUtilities::SupportsExport(BlendMode, MatProp);
+	if (bIsValidProperty)
 	{
-		PRAGMA_DISABLE_DEPRECATION_WARNINGS
-		// render the material to a texture to export as a bmp
-		if (!FMaterialUtilities::ExportMaterialProperty(Material, MatProp, OutputBMP, OutSize ))
+		FMeshData MeshSettings;
+		MeshSettings.MeshDescription = nullptr;
+		MeshSettings.TextureCoordinateBox = FBox2D(FVector2D(0.0f, 0.0f), FVector2D(1.0f, 1.0f));
+		MeshSettings.TextureCoordinateIndex = 0;
+
+		FMaterialData MaterialSettings;
+		MaterialSettings.Material = Material;
+
+		// Add all user defined properties for baking out
+		FPropertyEntry Entry(MatProp);
+		if (!Entry.bUseConstantValue && Material->IsPropertyActive(Entry.Property) && Entry.Property != MP_MAX)
+		{
+			FIntPoint TextureSize = FMaterialUtilities::FindMaxTextureSize(Material, MatProp);
+			MaterialSettings.PropertySizes.Add(Entry.Property, TextureSize);
+		}
+
+		TArray<FMeshData*> MeshSettingPtrs{ &MeshSettings };
+		TArray<FMaterialData*> MaterialSettingPtrs{ &MaterialSettings };
+
+		IMaterialBakingModule& Module = FModuleManager::Get().LoadModuleChecked<IMaterialBakingModule>("MaterialBaking");
+
+		bool bSetData = false;
+		bool bSetSize = false;
+		TArray<FBakeOutput> BakeOutputs;
+		Module.BakeMaterials(MaterialSettingPtrs, MeshSettingPtrs, BakeOutputs);
+		if (BakeOutputs.Num() > 0)
 		{
-			bIsValidMaterial = false;
-		}		
-		PRAGMA_ENABLE_DEPRECATION_WARNINGS
+			FBakeOutput& Output = BakeOutputs[0];
+			if (TArray<FColor>* PropertyData = Output.PropertyData.Find(MatProp))
+			{
+				OutputBMP = MoveTemp(*PropertyData);
+				bSetData = true;
+			}
+
+			if (FIntPoint* PropertySize = Output.PropertySizes.Find(MatProp))
+			{
+				OutSize = *PropertySize;
+				bSetSize = true;
+			}
+		}
+
+		bIsValidProperty = bSetData && bSetSize;
 	}
 
 	// make invalid textures a solid red
-	if (!bIsValidMaterial)
+	if (!bIsValidProperty)
 	{
 		OutSize = FIntPoint(1, 1);
 		OutputBMP.Empty();
@@ -1068,14 +1106,11 @@ void ExportOBJs(FOutputDevice& FileAr, FStringOutputDevice* MemAr, FFeedbackCont
 	// export extra material info if we added any
 	if (Materials)
 	{
-		// stop the rendering thread so we can easily render to texture
-		SCOPED_SUSPEND_RENDERING_THREAD(true);
-
 		// make a .MTL file next to the .obj file that contains the materials
 		FString MaterialLibFilename = FPaths::GetBaseFilename(OBJFilename, false) + TEXT(".mtl");
 
 		// use the output device file, just like the Exporter makes for the .obj, no backup
-		FOutputDeviceFile* MaterialLib = new FOutputDeviceFile(*MaterialLibFilename, true);
+		TUniquePtr<FOutputDeviceFile> MaterialLib = MakeUnique<FOutputDeviceFile>(*MaterialLibFilename, true);
 		MaterialLib->SetSuppressEventTag(true);
 		MaterialLib->SetAutoEmitLineTerminator(false);
 
@@ -1090,28 +1125,25 @@ void ExportOBJs(FOutputDevice& FileAr, FStringOutputDevice* MemAr, FFeedbackCont
 			MaterialLib->Logf(TEXT("newmtl %s\r\n"), *MaterialName);
 
 			{
-				FString BMPFilename = FPaths::GetPath(MaterialLibFilename) / MaterialName + TEXT("_D.bmp");
+				FString BMPFilename = FPaths::Combine(FPaths::GetPath(MaterialLibFilename), MaterialName + TEXT("_D.bmp"));
 				ExportMaterialPropertyTexture(BMPFilename, Material, MP_BaseColor);
 				MaterialLib->Logf(TEXT("\tmap_Kd %s\r\n"), *FPaths::GetCleanFilename(BMPFilename));
 			}
 
 			{
-				FString BMPFilename = FPaths::GetPath(MaterialLibFilename) / MaterialName + TEXT("_S.bmp");
+				FString BMPFilename = FPaths::Combine(FPaths::GetPath(MaterialLibFilename), MaterialName + TEXT("_S.bmp"));
 				ExportMaterialPropertyTexture(BMPFilename, Material, MP_Specular);
 				MaterialLib->Logf(TEXT("\tmap_Ks %s\r\n"), *FPaths::GetCleanFilename(BMPFilename));
 			}
 
 			{
-				FString BMPFilename = FPaths::GetPath(MaterialLibFilename) / MaterialName + TEXT("_N.bmp");
+				FString BMPFilename = FPaths::Combine(FPaths::GetPath(MaterialLibFilename), MaterialName + TEXT("_N.bmp"));
 				ExportMaterialPropertyTexture(BMPFilename, Material, MP_Normal);
 				MaterialLib->Logf(TEXT("\tbump %s\r\n"), *FPaths::GetCleanFilename(BMPFilename));
 			}
 
 			MaterialLib->Logf(TEXT("\r\n"));
 		}
-		
-		MaterialLib->TearDown();
-		delete MaterialLib;
 
 		Ar.Logf(TEXT("mtllib %s\n"), *FPaths::GetCleanFilename(MaterialLibFilename));
 	}
diff --git a/Engine/Source/Editor/UnrealEd/UnrealEd.Build.cs b/Engine/Source/Editor/UnrealEd/UnrealEd.Build.cs
--- a/Engine/Source/Editor/UnrealEd/UnrealEd.Build.cs
+++ b/Engine/Source/Editor/UnrealEd/UnrealEd.Build.cs
@@ -195,7 +195,8 @@ public class UnrealEd : ModuleRules
 				"EditorInteractiveToolsFramework",
 				"TraceLog",
 				"DeveloperSettings",
-				"AnimationModifiers"
+				"AnimationModifiers",
+				"MaterialBaking",
 			}
 		);
 
