Static analysis fix - possible out-of-bounds array access

diff --git a/Engine/Source/Developer/MaterialUtilities/Public/MaterialUtilities.h b/Engine/Source/Developer/MaterialUtilities/Public/MaterialUtilities.h
--- a/Engine/Source/Developer/MaterialUtilities/Public/MaterialUtilities.h
+++ b/Engine/Source/Developer/MaterialUtilities/Public/MaterialUtilities.h
@@ -87,18 +87,18 @@ struct FFlattenMaterial
 		}
 	}
 
-	const bool DoesPropertyContainData(const EFlattenMaterialProperties Property) const { return PropertySamples[(int32)Property].Num() > 0; }
+	const bool DoesPropertyContainData(const EFlattenMaterialProperties Property) const { return GetSamplesEntry(Property).Num() > 0; }
 
-	const bool IsPropertyConstant(const EFlattenMaterialProperties Property) const { return PropertySamples[(int32)Property].Num() == 1; }
+	const bool IsPropertyConstant(const EFlattenMaterialProperties Property) const { return GetSamplesEntry(Property).Num() == 1; }
 
-	const bool ShouldGenerateDataForProperty(const EFlattenMaterialProperties Property) const { return PropertySizes[(int32)Property].GetMin() > 0; }
+	const bool ShouldGenerateDataForProperty(const EFlattenMaterialProperties Property) const { return GetSizesEntry(Property).GetMin() > 0; }
 
-	const FIntPoint GetPropertySize(const EFlattenMaterialProperties Property) const{ return PropertySizes[(int32)Property]; }
-	void SetPropertySize(const EFlattenMaterialProperties Property, const FIntPoint& InSize) { PropertySizes[(int32)Property] = InSize; }
+	const FIntPoint GetPropertySize(const EFlattenMaterialProperties Property) const{ return GetSizesEntry(Property); }
+	void SetPropertySize(const EFlattenMaterialProperties Property, const FIntPoint& InSize) { GetSizesEntry(Property) = InSize; }
+
+	TArray<FColor>& GetPropertySamples(const EFlattenMaterialProperties Property) { return GetSamplesEntry(Property); }
+	const TArray<FColor>& GetPropertySamples(const EFlattenMaterialProperties Property) const { return GetSamplesEntry(Property); }
 
-	TArray<FColor>& GetPropertySamples(const EFlattenMaterialProperties Property) { return PropertySamples[(int32)Property]; }
-	const TArray<FColor>& GetPropertySamples(const EFlattenMaterialProperties Property) const { return PropertySamples[(int32)Property]; }
-	
 	/** Material Guid */
 	FGuid			MaterialId;	
 	FIntPoint		RenderSize;
@@ -115,6 +115,42 @@ struct FFlattenMaterial
 	int32			UVChannel;
 
 private:
+	FIntPoint& GetSizesEntry(const EFlattenMaterialProperties Property)
+	{
+		const int32 Index = (int32)Property;
+		check(Index < (int32)EFlattenMaterialProperties::NumFlattenMaterialProperties);
+
+		static const FIntPoint TempSize;
+		return PropertySizes.IsValidIndex(Index) ? PropertySizes[Index] : TempSize;
+	}
+
+	TArray<FColor>& GetSamplesEntry(const EFlattenMaterialProperties Property)
+	{
+		const int32 Index = (int32)Property;
+		check(Index < (int32)EFlattenMaterialProperties::NumFlattenMaterialProperties);
+
+		static const TArray<FColor> TempArray;
+		return PropertySamples.IsValidIndex(Index) ? PropertySamples[Index] : TempArray;
+	}
+
+	const FIntPoint& GetSizesEntry(const EFlattenMaterialProperties Property) const
+	{
+		const int32 Index = (int32)Property;
+		check(Index < (int32)EFlattenMaterialProperties::NumFlattenMaterialProperties);
+
+		static const FIntPoint TempSize;
+		return PropertySizes.IsValidIndex(Index) ? PropertySizes[Index] : TempSize;
+	}
+
+	const TArray<FColor>& GetSamplesEntry(const EFlattenMaterialProperties Property) const
+	{
+		const int32 Index = (int32)Property;
+		check(Index < (int32)EFlattenMaterialProperties::NumFlattenMaterialProperties);
+
+		static const TArray<FColor> TempArray;
+		return PropertySamples.IsValidIndex(Index) ? PropertySamples[Index] : TempArray;
+	}
+
 	/** Texture sizes for each individual property*/
 	FIntPoint PropertySizes[(uint32)EFlattenMaterialProperties::NumFlattenMaterialProperties];
 	/** Baked down texture samples for each individual property*/
@@ -226,7 +262,7 @@ public:
 	static FIntPoint FindMaxTextureSize(UMaterialInterface* InMaterialInterface, FIntPoint MinimumSize = FIntPoint(1, 1));
 
 	/**
-	 * Whether material utilities support exporting specified material blend mode and property 
+	 * Whether material utilities support exporting specified material blend mode and property
 	 */
 	static bool SupportsExport(EBlendMode InBlendMode, EMaterialProperty InMaterialProperty);
 
@@ -237,9 +273,9 @@ public:
 	 * @param InMaterial			Target material
 	 * @param InMaterialProperty	Material property to render
 	 * @param InRenderTarget		Render target to render to
-	 * @param OutBMP				Output array of rendered samples 
+	 * @param OutBMP				Output array of rendered samples
 	 * @return						Whether operation was successful
-	 */	
+	 */
 	UE_DEPRECATED(4.11, "Please use ExportMaterialProperty function with new signature")
 	static bool ExportMaterialProperty(UWorld* InWorld, UMaterialInterface* InMaterial, EMaterialProperty InMaterialProperty, UTextureRenderTarget2D* InRenderTarget, TArray<FColor>& OutBMP);
 
@@ -247,7 +283,7 @@ public:
 	* Renders specified material property into texture
 	*
 	* @param InMaterial			Target material
-	* @param InMaterialProperty	Material property to render	
+	* @param InMaterialProperty	Material property to render
 	* @param OutBMP				Output array of rendered samples
 	* @param OutSize			Output size of the rendered samples
 	* @return					Whether operation was successful
@@ -261,7 +297,7 @@ public:
 	* @param InMaterial			Target material
 	* @param InMaterialProperty	Material property to render
 	* @param InSize				Input size for the rendered samples
-	* @param OutBMP				Output array of rendered samples	
+	* @param OutBMP				Output array of rendered samples
 	* @return					Whether operation was successful
 	*/
 	UE_DEPRECATED(4.19, "Please use new functionality in MaterialBaking module")
@@ -273,7 +309,7 @@ public:
 	 * @param InWorld				World object to use for material property rendering
 	 * @param InMaterial			Target material
 	 * @param InMaterialProperty	Material property to render
-	 * @param OutBMP				Output array of rendered samples 
+	 * @param OutBMP				Output array of rendered samples
 	 * @return						Whether operation was successful
 	 */
 	UE_DEPRECATED(4.11, "Please use ExportMaterialProperty function with new signature")
@@ -289,7 +325,7 @@ public:
 	 */
 	UE_DEPRECATED(4.11, "Please use ExportMaterial function with new signature")
 	static bool ExportMaterial(UWorld* InWorld, UMaterialInterface* InMaterial, FFlattenMaterial& OutFlattenMaterial);
-	
+
 	/**
 	* Flattens specified material
 	*
@@ -312,9 +348,9 @@ public:
 	 * @return						Whether operation was successful
 	 */
 	static bool ExportLandscapeMaterial(ALandscapeProxy* InLandscape, const TSet<FPrimitiveComponentId>& HiddenPrimitives, FFlattenMaterial& OutFlattenMaterial);
-	
+
 	/**
- 	 * Generates a texture from an array of samples 
+ 	 * Generates a texture from an array of samples
 	 *
 	 * @param Outer					Outer for the material and texture objects, if NULL new packages will be created for each asset
 	 * @param AssetLongName			Long asset path for the new texture
@@ -340,7 +376,7 @@ public:
 	 */
 	static UMaterial* CreateMaterial(const FFlattenMaterial& InFlattenMaterial, UPackage* InOuter, const FString& BaseName, EObjectFlags Flags, const struct FMaterialProxySettings& MaterialProxySettings, TArray<UObject*>& OutGeneratedAssets, const TextureGroup& InTextureGroup = TEXTUREGROUP_World);
 
-	/** 
+	/**
 	* Creates an instanced material based of BaseMaterial
 	* @param Outer					Outer for the material and texture objects, if NULL new packages will be created for each asset
 	* @param BaseName				BaseName for the material and texture objects, should be a long package name in case Outer is not specified
@@ -411,7 +447,7 @@ public:
 	static void OptimizeSampleArray(TArray<FColor>& InSamples, FIntPoint& InSampleSize);
 
 	/**
-	* Contains errors generated when exporting material texcoord scales. 
+	* Contains errors generated when exporting material texcoord scales.
 	* Used to prevent displaying duplicates, as instances using the same shaders get the same issues.
 	*/
 	class MATERIALUTILITIES_API FExportErrorManager
@@ -470,7 +506,7 @@ public:
 	* @param InMaterial			Target material
 	* @param QualityLevel		Quality level used for the shader profiling.
 	* @param FeatureLevel		Feature level used for the shader profiling.
-	* @param OutErrors			Manager to log errors (removes duplicates and similar errors)	
+	* @param OutErrors			Manager to log errors (removes duplicates and similar errors)
 	* @return					Whether operation was successful
 	*/
 	static bool ExportMaterialUVDensities(UMaterialInterface* InMaterial, EMaterialQualityLevel::Type QualityLevel, ERHIFeatureLevel::Type FeatureLevel, FExportErrorManager& OutErrors);
@@ -486,11 +522,11 @@ public:
 	/** Creates a proxy material and the required texture assets */
 	static UMaterialInterface* CreateProxyMaterialAndTextures(const FString& PackageName, const FString& AssetName, const FBakeOutput& BakeOutput, const FMeshData& MeshData, const FMaterialData& MaterialData, UMaterialOptions* Options);
 private:
-	
+
 	/**
 	* Private export material function to which unique signatures are mapped for unified code path
 	*
-	* @param InMaterialData			Target material data 
+	* @param InMaterialData			Target material data
 	* @param OutFlattenMaterial		Output flattened material
 	* @return						Whether operation was successful
 	*/
@@ -521,9 +557,9 @@ private:
 	* @return						Created render target
 	*/
 	static UTextureRenderTarget2D* CreateRenderTarget(bool bInForceLinearGamma, bool bNormalMap, EPixelFormat InPixelFormat, FIntPoint& InTargetSize );
-	
+
 	/** Clears the pool with available render targets */
-	static void ClearRenderTargetPool();	
+	static void ClearRenderTargetPool();
 
 	/** Call back for garbage collector, cleans up the RenderTargetPool if CurrentlyRendering is set to false */
 	void OnPreGarbageCollect();
