Static analysis fix - possible out-of-bounds array access

diff --git a/Engine/Source/Developer/MaterialUtilities/Public/MaterialUtilities.h b/Engine/Source/Developer/MaterialUtilities/Public/MaterialUtilities.h
--- a/Engine/Source/Developer/MaterialUtilities/Public/MaterialUtilities.h
+++ b/Engine/Source/Developer/MaterialUtilities/Public/MaterialUtilities.h
@@ -87,18 +87,18 @@ struct FFlattenMaterial
 		}
 	}
 
-	const bool DoesPropertyContainData(const EFlattenMaterialProperties Property) const { return PropertySamples[(int32)Property].Num() > 0; }
+	const bool DoesPropertyContainData(const EFlattenMaterialProperties Property) const { return GetSamplesEntry(Property).Num() > 0; }
 
-	const bool IsPropertyConstant(const EFlattenMaterialProperties Property) const { return PropertySamples[(int32)Property].Num() == 1; }
+	const bool IsPropertyConstant(const EFlattenMaterialProperties Property) const { return GetSamplesEntry(Property).Num() == 1; }
 
-	const bool ShouldGenerateDataForProperty(const EFlattenMaterialProperties Property) const { return PropertySizes[(int32)Property].GetMin() > 0; }
+	const bool ShouldGenerateDataForProperty(const EFlattenMaterialProperties Property) const { return GetSizesEntry(Property).GetMin() > 0; }
 
-	const FIntPoint GetPropertySize(const EFlattenMaterialProperties Property) const{ return PropertySizes[(int32)Property]; }
-	void SetPropertySize(const EFlattenMaterialProperties Property, const FIntPoint& InSize) { PropertySizes[(int32)Property] = InSize; }
+	const FIntPoint GetPropertySize(const EFlattenMaterialProperties Property) const{ return GetSizesEntry(Property); }
+	void SetPropertySize(const EFlattenMaterialProperties Property, const FIntPoint& InSize) { GetSizesEntry(Property) = InSize; }
+
+	TArray<FColor>& GetPropertySamples(const EFlattenMaterialProperties Property) { return GetSamplesEntry(Property); }
+	const TArray<FColor>& GetPropertySamples(const EFlattenMaterialProperties Property) const { return GetSamplesEntry(Property); }
 
-	TArray<FColor>& GetPropertySamples(const EFlattenMaterialProperties Property) { return PropertySamples[(int32)Property]; }
-	const TArray<FColor>& GetPropertySamples(const EFlattenMaterialProperties Property) const { return PropertySamples[(int32)Property]; }
-	
 	/** Material Guid */
 	FGuid			MaterialId;	
 	FIntPoint		RenderSize;
@@ -115,6 +115,42 @@ struct FFlattenMaterial
 	int32			UVChannel;
 
 private:
+	FIntPoint& GetSizesEntry(const EFlattenMaterialProperties Property)
+	{
+		const uint32 Index = (uint32)Property;
+		check(Index < (uint32)EFlattenMaterialProperties::NumFlattenMaterialProperties);
+
+		static FIntPoint TempSize;
+		return Index < (uint32)EFlattenMaterialProperties::NumFlattenMaterialProperties ? PropertySizes[Index] : TempSize;
+	}
+
+	TArray<FColor>& GetSamplesEntry(const EFlattenMaterialProperties Property)
+	{
+		const uint32 Index = (uint32)Property;
+		check(Index < (uint32)EFlattenMaterialProperties::NumFlattenMaterialProperties);
+
+		static TArray<FColor> TempArray;
+		return Index < (uint32)EFlattenMaterialProperties::NumFlattenMaterialProperties ? PropertySamples[Index] : TempArray;
+	}
+
+	const FIntPoint& GetSizesEntry(const EFlattenMaterialProperties Property) const
+	{
+		const uint32 Index = (uint32)Property;
+		check(Index < (uint32)EFlattenMaterialProperties::NumFlattenMaterialProperties);
+
+		static const FIntPoint TempSize;
+		return Index < (uint32)EFlattenMaterialProperties::NumFlattenMaterialProperties ? PropertySizes[Index] : TempSize;
+	}
+
+	const TArray<FColor>& GetSamplesEntry(const EFlattenMaterialProperties Property) const
+	{
+		const uint32 Index = (uint32)Property;
+		check(Index < (uint32)EFlattenMaterialProperties::NumFlattenMaterialProperties);
+
+		static const TArray<FColor> TempArray;
+		return Index < (uint32)EFlattenMaterialProperties::NumFlattenMaterialProperties ? PropertySamples[Index] : TempArray;
+	}
+
 	/** Texture sizes for each individual property*/
 	FIntPoint PropertySizes[(uint32)EFlattenMaterialProperties::NumFlattenMaterialProperties];
 	/** Baked down texture samples for each individual property*/
