Fixes motion blur lengths when using time dilation

diff --git a/Engine/Plugins/Animation/ControlRig/Source/ControlRigEditor/Private/Tools/ControlRigPoseThumbnailRenderer.cpp b/Engine/Plugins/Animation/ControlRig/Source/ControlRigEditor/Private/Tools/ControlRigPoseThumbnailRenderer.cpp
--- a/Engine/Plugins/Animation/ControlRig/Source/ControlRigEditor/Private/Tools/ControlRigPoseThumbnailRenderer.cpp
+++ b/Engine/Plugins/Animation/ControlRig/Source/ControlRigEditor/Private/Tools/ControlRigPoseThumbnailRenderer.cpp
@@ -210,7 +210,7 @@ void UControlRigPoseThumbnailRenderer::Draw(UObject* Object, int32 X, int32 Y, u
 		ThumbnailScene->SetControlRigPoseAsset(PoseAsset);
 
 		FSceneViewFamilyContext ViewFamily(FSceneViewFamily::ConstructionValues(RenderTarget, ThumbnailScene->GetScene(), FEngineShowFlags(ESFIM_Game))
-			.SetWorldTimes(FApp::GetCurrentTime() - GStartTime, FApp::GetDeltaTime(), FApp::GetCurrentTime() - GStartTime)
+			.SetTime(UThumbnailRenderer::GetTime())
 			.SetAdditionalViewFamily(bAdditionalViewFamily));
 
 		ViewFamily.EngineShowFlags.DisableAdvancedFeatures();
diff --git a/Engine/Plugins/Enterprise/LidarPointCloud/Source/LidarPointCloudRuntime/Private/LidarPointCloudLODManager.cpp b/Engine/Plugins/Enterprise/LidarPointCloud/Source/LidarPointCloudRuntime/Private/LidarPointCloudLODManager.cpp
--- a/Engine/Plugins/Enterprise/LidarPointCloud/Source/LidarPointCloudRuntime/Private/LidarPointCloudLODManager.cpp
+++ b/Engine/Plugins/Enterprise/LidarPointCloud/Source/LidarPointCloudRuntime/Private/LidarPointCloudLODManager.cpp
@@ -152,7 +152,7 @@ bool FLidarPointCloudViewData::ComputeFromEditorViewportClient(FViewportClient*
 		if (Client->Viewport && Client->Viewport->GetSizeXY() != FIntPoint::ZeroValue)
 		{
 			FSceneViewFamily::ConstructionValues CVS(nullptr, nullptr, FEngineShowFlags(EShowFlagInitMode::ESFIM_Game));
-			CVS.SetWorldTimes(0, 0, 0);
+			CVS.SetTime(FGameTime());
 			FSceneViewFamily ViewFamily(CVS);
 			FSceneView* View = Client->CalcSceneView(&ViewFamily);
 
diff --git a/Engine/Plugins/Enterprise/VariantManagerContent/Source/VariantManagerContent/Private/ThumbnailGenerator.cpp b/Engine/Plugins/Enterprise/VariantManagerContent/Source/VariantManagerContent/Private/ThumbnailGenerator.cpp
--- a/Engine/Plugins/Enterprise/VariantManagerContent/Source/VariantManagerContent/Private/ThumbnailGenerator.cpp
+++ b/Engine/Plugins/Enterprise/VariantManagerContent/Source/VariantManagerContent/Private/ThumbnailGenerator.cpp
@@ -41,7 +41,7 @@ namespace ThumbnailGeneratorImpl
 
 		FSceneViewFamilyContext ViewFamily(
 			FSceneViewFamily::ConstructionValues(RenderTargetResource, Scene, FEngineShowFlags(ESFIM_Game))
-			.SetWorldTimes(FApp::GetCurrentTime() - GStartTime, FApp::GetDeltaTime(), FApp::GetCurrentTime() - GStartTime)
+			.SetTime(FGameTime::GetTimeSinceAppStart())
 		);
 
 		ViewFamily.SetScreenPercentageInterface(new FLegacyScreenPercentageDriver(
@@ -57,7 +57,7 @@ namespace ThumbnailGeneratorImpl
 		FSceneView* NewView = new FSceneView(ViewInitOptions);
 		ViewFamily.Views.Add(NewView);
 
-		FCanvas Canvas(RenderTargetResource, NULL, FApp::GetCurrentTime() - GStartTime, FApp::GetDeltaTime(), FApp::GetCurrentTime() - GStartTime, Scene->GetFeatureLevel());
+		FCanvas Canvas(RenderTargetResource, NULL, FGameTime::GetTimeSinceAppStart(), Scene->GetFeatureLevel());
 		Canvas.Clear(FLinearColor::Transparent);
 		GetRendererModule().BeginRenderingViewFamily(&Canvas, &ViewFamily);
 
diff --git a/Engine/Plugins/Experimental/AutomationUtils/Source/AutomationUtils/Private/AutomationUtilsBlueprintLibrary.cpp b/Engine/Plugins/Experimental/AutomationUtils/Source/AutomationUtils/Private/AutomationUtilsBlueprintLibrary.cpp
--- a/Engine/Plugins/Experimental/AutomationUtils/Source/AutomationUtils/Private/AutomationUtilsBlueprintLibrary.cpp
+++ b/Engine/Plugins/Experimental/AutomationUtils/Source/AutomationUtils/Private/AutomationUtilsBlueprintLibrary.cpp
@@ -45,9 +45,7 @@ public:
 		ShowFlags.SetBloom(0);
 
 		// Turn off time the ultimate source of noise.
-		InViewFamily.CurrentWorldTime = 0;
-		InViewFamily.CurrentRealTime = 0;
-		InViewFamily.DeltaWorldTime = 0;
+		InViewFamily.Time = FGameTime();
 	}
 
 	bool IsActiveThisFrame_Internal(const FSceneViewExtensionContext&) const
diff --git a/Engine/Plugins/Experimental/GPULightmass/Source/GPULightmass/Private/LightmapRenderer.cpp b/Engine/Plugins/Experimental/GPULightmass/Source/GPULightmass/Private/LightmapRenderer.cpp
--- a/Engine/Plugins/Experimental/GPULightmass/Source/GPULightmass/Private/LightmapRenderer.cpp
+++ b/Engine/Plugins/Experimental/GPULightmass/Source/GPULightmass/Private/LightmapRenderer.cpp
@@ -757,7 +757,7 @@ void FSceneRenderState::SetupRayTracingScene(int32 LODIndex)
 		nullptr,
 		nullptr,
 		FEngineShowFlags(ESFIM_Game))
-		.SetWorldTimes(0, 0, 0)
+		.SetTime(FGameTime())
 		.SetGammaCorrection(1.0f));
 
 	const FIntRect ViewRect(FIntPoint(0, 0), FIntPoint(GPreviewLightmapPhysicalTileSize, GPreviewLightmapPhysicalTileSize));
diff --git a/Engine/Plugins/Experimental/GeometryCache/Source/GeometryCacheEd/Private/GeometryCacheThumbnailRenderer.cpp b/Engine/Plugins/Experimental/GeometryCache/Source/GeometryCacheEd/Private/GeometryCacheThumbnailRenderer.cpp
--- a/Engine/Plugins/Experimental/GeometryCache/Source/GeometryCacheEd/Private/GeometryCacheThumbnailRenderer.cpp
+++ b/Engine/Plugins/Experimental/GeometryCache/Source/GeometryCacheEd/Private/GeometryCacheThumbnailRenderer.cpp
@@ -27,7 +27,7 @@ void UGeometryCacheThumbnailRenderer::Draw(UObject* Object, int32 X, int32 Y, ui
 		ThumbnailScene->GetScene()->UpdateSpeedTreeWind(0.0);
 
 		FSceneViewFamilyContext ViewFamily(FSceneViewFamily::ConstructionValues(RenderTarget, ThumbnailScene->GetScene(), FEngineShowFlags(ESFIM_Game))
-			.SetWorldTimes(FApp::GetCurrentTime() - GStartTime, FApp::GetDeltaTime(), FApp::GetCurrentTime() - GStartTime)
+			.SetTime(UThumbnailRenderer::GetTime())
 			.SetAdditionalViewFamily(bAdditionalViewFamily));
 
 		ViewFamily.EngineShowFlags.DisableAdvancedFeatures();
diff --git a/Engine/Plugins/Experimental/GeometryCollectionPlugin/Source/GeometryCollectionEditor/Private/GeometryCollection/GeometryCollectionThumbnailRenderer.cpp b/Engine/Plugins/Experimental/GeometryCollectionPlugin/Source/GeometryCollectionEditor/Private/GeometryCollection/GeometryCollectionThumbnailRenderer.cpp
--- a/Engine/Plugins/Experimental/GeometryCollectionPlugin/Source/GeometryCollectionEditor/Private/GeometryCollection/GeometryCollectionThumbnailRenderer.cpp
+++ b/Engine/Plugins/Experimental/GeometryCollectionPlugin/Source/GeometryCollectionEditor/Private/GeometryCollection/GeometryCollectionThumbnailRenderer.cpp
@@ -27,7 +27,7 @@ void UGeometryCollectionThumbnailRenderer::Draw(UObject* Object, int32 X, int32
 		ThumbnailScene->GetScene()->UpdateSpeedTreeWind(0.0);
 
 		FSceneViewFamilyContext ViewFamily(FSceneViewFamily::ConstructionValues(RenderTarget, ThumbnailScene->GetScene(), FEngineShowFlags(ESFIM_Game))
-			.SetWorldTimes(FApp::GetCurrentTime() - GStartTime, FApp::GetDeltaTime(), FApp::GetCurrentTime() - GStartTime)
+			.SetTime(UThumbnailRenderer::GetTime())
 			.SetAdditionalViewFamily(bAdditionalViewFamily));
 
 		ViewFamily.EngineShowFlags.DisableAdvancedFeatures();
diff --git a/Engine/Plugins/FX/Niagara/Source/Niagara/Private/NiagaraGpuComputeDispatch.cpp b/Engine/Plugins/FX/Niagara/Source/Niagara/Private/NiagaraGpuComputeDispatch.cpp
--- a/Engine/Plugins/FX/Niagara/Source/Niagara/Private/NiagaraGpuComputeDispatch.cpp
+++ b/Engine/Plugins/FX/Niagara/Source/Niagara/Private/NiagaraGpuComputeDispatch.cpp
@@ -458,7 +458,7 @@ void FNiagaraGpuComputeDispatch::ProcessPendingTicksFlush(FRHICommandListImmedia
 			FMemMark Mark(FMemStack::Get());
 
 			FSceneViewFamily ViewFamily(FSceneViewFamily::ConstructionValues(nullptr, nullptr, FEngineShowFlags(ESFIM_Game))
-				.SetWorldTimes(0, 0, 0)
+				.SetTime(FGameTime())
 				.SetGammaCorrection(1.0f));
 
 			FSceneViewInitOptions ViewInitOptions;
diff --git a/Engine/Plugins/FX/Niagara/Source/Niagara/Private/NiagaraRendererMeshes.cpp b/Engine/Plugins/FX/Niagara/Source/Niagara/Private/NiagaraRendererMeshes.cpp
--- a/Engine/Plugins/FX/Niagara/Source/Niagara/Private/NiagaraRendererMeshes.cpp
+++ b/Engine/Plugins/FX/Niagara/Source/Niagara/Private/NiagaraRendererMeshes.cpp
@@ -652,7 +652,7 @@ FNiagaraMeshUniformBufferRef FNiagaraRendererMeshes::CreatePerViewUniformBuffer(
 
 	const bool bUseLocalSpace = UseLocalSpace(&SceneProxy);
 	PerViewUniformParameters.bLocalSpace = bUseLocalSpace;
-	PerViewUniformParameters.DeltaSeconds = View.Family->DeltaWorldTime;
+	PerViewUniformParameters.DeltaSeconds = View.Family->Time.GetDeltaWorldTimeSeconds();
 	PerViewUniformParameters.MeshScale = MeshData.Scale;
 	PerViewUniformParameters.SystemLWCTile = SceneProxy.GetLWCRenderTile();
 
diff --git a/Engine/Plugins/FX/Niagara/Source/Niagara/Private/NiagaraRendererRibbons.cpp b/Engine/Plugins/FX/Niagara/Source/Niagara/Private/NiagaraRendererRibbons.cpp
--- a/Engine/Plugins/FX/Niagara/Source/Niagara/Private/NiagaraRendererRibbons.cpp
+++ b/Engine/Plugins/FX/Niagara/Source/Niagara/Private/NiagaraRendererRibbons.cpp
@@ -1384,7 +1384,7 @@ void FNiagaraRendererRibbons::CreatePerViewResources(
 
 	bool bUseLocalSpace = UseLocalSpace(SceneProxy);
 	PerViewUniformParameters.bLocalSpace = bUseLocalSpace;
-	PerViewUniformParameters.DeltaSeconds = ViewFamily.DeltaWorldTime;
+	PerViewUniformParameters.DeltaSeconds = ViewFamily.Time.GetDeltaWorldTimeSeconds();
 	PerViewUniformParameters.SystemLWCTile = SceneProxy->GetLWCRenderTile();
 	PerViewUniformParameters.CameraUp = View->GetViewUp(); // FVector4(0.0f, 0.0f, 1.0f, 0.0f);
 	PerViewUniformParameters.CameraRight = View->GetViewRight();//	FVector4(1.0f, 0.0f, 0.0f, 0.0f);
diff --git a/Engine/Plugins/FX/Niagara/Source/Niagara/Private/NiagaraRendererSprites.cpp b/Engine/Plugins/FX/Niagara/Source/Niagara/Private/NiagaraRendererSprites.cpp
--- a/Engine/Plugins/FX/Niagara/Source/Niagara/Private/NiagaraRendererSprites.cpp
+++ b/Engine/Plugins/FX/Niagara/Source/Niagara/Private/NiagaraRendererSprites.cpp
@@ -493,7 +493,7 @@ FNiagaraSpriteUniformBufferRef FNiagaraRendererSprites::CreateViewUniformBuffer(
 	PerViewUniformParameters.RotationBias = 0.0f;
 	PerViewUniformParameters.RotationScale = 1.0f;
 	PerViewUniformParameters.TangentSelector = FVector4f(0.0f, 0.0f, 0.0f, 1.0f);
-	PerViewUniformParameters.DeltaSeconds = ViewFamily.DeltaWorldTime;
+	PerViewUniformParameters.DeltaSeconds = ViewFamily.Time.GetDeltaWorldTimeSeconds();
 	PerViewUniformParameters.NormalsType = 0.0f;
 	PerViewUniformParameters.NormalsSphereCenter = FVector4f(0.0f, 0.0f, 0.0f, 1.0f);
 	PerViewUniformParameters.NormalsCylinderUnitDirection = FVector4f(0.0f, 0.0f, 1.0f, 0.0f);
diff --git a/Engine/Plugins/FX/Niagara/Source/NiagaraEditor/Private/NiagaraBakerRenderer.cpp b/Engine/Plugins/FX/Niagara/Source/NiagaraEditor/Private/NiagaraBakerRenderer.cpp
--- a/Engine/Plugins/FX/Niagara/Source/NiagaraEditor/Private/NiagaraBakerRenderer.cpp
+++ b/Engine/Plugins/FX/Niagara/Source/NiagaraEditor/Private/NiagaraBakerRenderer.cpp
@@ -56,7 +56,7 @@ bool FNiagaraBakerRenderer::RenderView(UNiagaraComponent* PreviewComponent, cons
 		return false;
 	}
 
-	FCanvas Canvas(RenderTarget->GameThread_GetRenderTargetResource(), nullptr, WorldTime, FApp::GetDeltaTime(), WorldTime, World->Scene->GetFeatureLevel());
+	FCanvas Canvas(RenderTarget->GameThread_GetRenderTargetResource(), nullptr, FGameTime::CreateUndilated(WorldTime, FApp::GetDeltaTime()), World->Scene->GetFeatureLevel());
 	Canvas.Clear(FLinearColor::Black);
 
 	bool bRendered = false;
@@ -172,7 +172,7 @@ bool FNiagaraBakerRenderer::RenderView(UNiagaraComponent* PreviewComponent, cons
 			// Create View Family
 			FSceneViewFamilyContext ViewFamily(
 				FSceneViewFamily::ConstructionValues(RenderTarget->GameThread_GetRenderTargetResource(), World->Scene, FEngineShowFlags(ESFIM_Game))
-				.SetWorldTimes(WorldTime, FApp::GetDeltaTime(), WorldTime)
+				.SetTime(FGameTime::CreateUndilated(WorldTime, FApp::GetDeltaTime()))
 				.SetGammaCorrection(GammaCorrection)
 			);
 
diff --git a/Engine/Plugins/FX/Niagara/Source/NiagaraVertexFactories/Public/NiagaraVertexFactory.h b/Engine/Plugins/FX/Niagara/Source/NiagaraVertexFactories/Public/NiagaraVertexFactory.h
--- a/Engine/Plugins/FX/Niagara/Source/NiagaraVertexFactories/Public/NiagaraVertexFactory.h
+++ b/Engine/Plugins/FX/Niagara/Source/NiagaraVertexFactories/Public/NiagaraVertexFactory.h
@@ -107,12 +107,12 @@ public:
 
 	bool CheckAndUpdateLastFrame(const FSceneViewFamily& ViewFamily, const FSceneView *View = nullptr) const
 	{
-		if (LastFrameSetup != MAX_uint32 && (&ViewFamily == LastViewFamily) && (View == LastView) && ViewFamily.FrameNumber == LastFrameSetup && LastFrameRealTime == ViewFamily.CurrentRealTime)
+		if (LastFrameSetup != MAX_uint32 && (&ViewFamily == LastViewFamily) && (View == LastView) && ViewFamily.FrameNumber == LastFrameSetup && LastFrameRealTime == ViewFamily.Time.GetRealTimeSeconds())
 		{
 			return false;
 		}
 		LastFrameSetup = ViewFamily.FrameNumber;
-		LastFrameRealTime = ViewFamily.CurrentRealTime;
+		LastFrameRealTime = ViewFamily.Time.GetRealTimeSeconds();
 		LastViewFamily = &ViewFamily;
 		LastView = View;
 		return true;
diff --git a/Engine/Plugins/MeshPainting/Source/MeshPaintingToolset/Private/MeshTexturePaintingTool.cpp b/Engine/Plugins/MeshPainting/Source/MeshPaintingToolset/Private/MeshTexturePaintingTool.cpp
--- a/Engine/Plugins/MeshPainting/Source/MeshPaintingToolset/Private/MeshTexturePaintingTool.cpp
+++ b/Engine/Plugins/MeshPainting/Source/MeshPaintingToolset/Private/MeshTexturePaintingTool.cpp
@@ -827,7 +827,7 @@ void UMeshTexturePaintingTool::PaintTexture(FMeshPaintParameters& InParams, TArr
 	check(BrushRenderTargetResource != nullptr);
 
 	// Create a canvas for the brush render target.
-	FCanvas BrushPaintCanvas(BrushRenderTargetResource, nullptr, 0, 0, 0, FeatureLevel);
+	FCanvas BrushPaintCanvas(BrushRenderTargetResource, nullptr, FGameTime(), FeatureLevel);
 
 	// Parameters for brush paint
 	TRefCountPtr< FMeshPaintBatchedElementParameters > MeshPaintBatchedElementParameters(new FMeshPaintBatchedElementParameters());
@@ -865,7 +865,7 @@ void UMeshTexturePaintingTool::PaintTexture(FMeshPaintParameters& InParams, TArr
 		check(BrushMaskRenderTargetResource != nullptr);
 
 		// Create a canvas for the brush mask rendertarget and clear it to black.
-		BrushMaskCanvas = TSharedPtr<FCanvas>(new FCanvas(BrushMaskRenderTargetResource, nullptr, 0, 0, 0, FeatureLevel));
+		BrushMaskCanvas = TSharedPtr<FCanvas>(new FCanvas(BrushMaskRenderTargetResource, nullptr, FGameTime(), FeatureLevel));
 		BrushMaskCanvas->Clear(FLinearColor::Black);
 
 		// Parameters for the mask
@@ -1063,7 +1063,7 @@ void UMeshTexturePaintingTool::PaintTexture(FMeshPaintParameters& InParams, TArr
 		// Dilate the paint stroke into the texture seams.
 		{
 			// Create a canvas for the render target.
-			FCanvas Canvas3(RenderTargetResource, nullptr, 0, 0, 0, FeatureLevel);
+			FCanvas Canvas3(RenderTargetResource, nullptr, FGameTime(), FeatureLevel);
 
 
 			TRefCountPtr< FMeshPaintDilateBatchedElementParameters > MeshPaintDilateBatchedElementParameters(new FMeshPaintDilateBatchedElementParameters());
diff --git a/Engine/Plugins/MeshPainting/Source/MeshPaintingToolset/Private/TexturePaintToolset.cpp b/Engine/Plugins/MeshPainting/Source/MeshPaintingToolset/Private/TexturePaintToolset.cpp
--- a/Engine/Plugins/MeshPainting/Source/MeshPaintingToolset/Private/TexturePaintToolset.cpp
+++ b/Engine/Plugins/MeshPainting/Source/MeshPaintingToolset/Private/TexturePaintToolset.cpp
@@ -28,7 +28,7 @@ void UTexturePaintToolset::CopyTextureToRenderTargetTexture(UTexture* SourceText
 	check(RenderTargetResource != nullptr);
 	
 	// Create a canvas for the render target and clear it to black
-	FCanvas Canvas(RenderTargetResource, nullptr, 0, 0, 0, FeatureLevel);
+	FCanvas Canvas(RenderTargetResource, nullptr, FGameTime(), FeatureLevel);
 
 	const uint32 Width = RenderTargetTexture->GetSurfaceWidth();
 	const uint32 Height = RenderTargetTexture->GetSurfaceHeight();
@@ -204,7 +204,7 @@ bool UTexturePaintToolset::GenerateSeamMask(UMeshComponent* MeshComponent, int32
 
 	{
 		// Create a canvas for the render target and clear it to white
-		FCanvas Canvas(RenderTargetResource, nullptr, 0, 0, 0, GEditor->GetEditorWorldContext().World()->FeatureLevel);
+		FCanvas Canvas(RenderTargetResource, nullptr, FGameTime(), GEditor->GetEditorWorldContext().World()->FeatureLevel);
 		Canvas.Clear(FLinearColor::White);
 
 		TArray<FCanvasUVTri> TriList;
diff --git a/Engine/Plugins/MovieScene/MovieRenderPipeline/Source/MovieRenderPipelineRenderPasses/Private/MoviePipelineImagePassBase.cpp b/Engine/Plugins/MovieScene/MovieRenderPipeline/Source/MovieRenderPipelineRenderPasses/Private/MoviePipelineImagePassBase.cpp
--- a/Engine/Plugins/MovieScene/MovieRenderPipeline/Source/MovieRenderPipelineRenderPasses/Private/MoviePipelineImagePassBase.cpp
+++ b/Engine/Plugins/MovieScene/MovieRenderPipeline/Source/MovieRenderPipelineRenderPasses/Private/MoviePipelineImagePassBase.cpp
@@ -77,7 +77,7 @@ TSharedPtr<FSceneViewFamilyContext> UMoviePipelineImagePassBase::CalculateViewFa
 		RenderTarget,
 		GetPipeline()->GetWorld()->Scene,
 		ShowFlags)
-		.SetWorldTimes(TimeData.WorldSeconds, TimeData.FrameDeltaTime, TimeData.WorldSeconds)
+		.SetTime(FGameTime::CreateUndilated(TimeData.WorldSeconds, TimeData.FrameDeltaTime))
 		.SetRealtimeUpdate(true));
 
 	OutViewFamily->SceneCaptureSource = InOutSampleState.SceneCaptureSource;
diff --git a/Engine/Plugins/Runtime/ApexDestruction/Source/ApexDestructionEditor/Private/DestructibleMeshThumbnailRenderer.cpp b/Engine/Plugins/Runtime/ApexDestruction/Source/ApexDestructionEditor/Private/DestructibleMeshThumbnailRenderer.cpp
--- a/Engine/Plugins/Runtime/ApexDestruction/Source/ApexDestructionEditor/Private/DestructibleMeshThumbnailRenderer.cpp
+++ b/Engine/Plugins/Runtime/ApexDestruction/Source/ApexDestructionEditor/Private/DestructibleMeshThumbnailRenderer.cpp
@@ -29,7 +29,7 @@ PRAGMA_DISABLE_DEPRECATION_WARNINGS
 
 		ThumbnailScene->SetDestructibleMesh(DestructibleMesh);
 		FSceneViewFamilyContext ViewFamily( FSceneViewFamily::ConstructionValues( RenderTarget, ThumbnailScene->GetScene(), FEngineShowFlags(ESFIM_Game) )
-			.SetWorldTimes(FApp::GetCurrentTime() - GStartTime, FApp::GetDeltaTime(), FApp::GetCurrentTime() - GStartTime)
+			.SetTime(UThumbnailRenderer::GetTime())
 			.SetAdditionalViewFamily(bAdditionalViewFamily));
 
 		ViewFamily.EngineShowFlags.DisableAdvancedFeatures();
diff --git a/Engine/Plugins/Runtime/HairStrands/Source/HairStrandsCore/Private/GroomDebug.cpp b/Engine/Plugins/Runtime/HairStrands/Source/HairStrandsCore/Private/GroomDebug.cpp
--- a/Engine/Plugins/Runtime/HairStrands/Source/HairStrandsCore/Private/GroomDebug.cpp
+++ b/Engine/Plugins/Runtime/HairStrands/Source/HairStrandsCore/Private/GroomDebug.cpp
@@ -906,7 +906,7 @@ void RunHairStrandsDebug(
 		PassParameters->RenderTargets[0] = FRenderTargetBinding(SceneColorTexture, ERenderTargetLoadAction::ELoad, 0);
 
 		const FSceneViewFamily& ViewFamily = *View.Family;
-		FCanvas& Canvas = *FCanvas::Create(GraphBuilder, SceneColorTexture, nullptr, ViewFamily.CurrentRealTime, ViewFamily.CurrentWorldTime, ViewFamily.DeltaWorldTime, View.FeatureLevel);
+		FCanvas& Canvas = *FCanvas::Create(GraphBuilder, SceneColorTexture, nullptr, ViewFamily.Time, View.FeatureLevel);
 		Canvas.SetRenderTargetRect(Viewport);
 
 		{
diff --git a/Engine/Plugins/Runtime/MeshModelingToolset/Source/ModelingComponents/Private/Scene/WorldRenderCapture.cpp b/Engine/Plugins/Runtime/MeshModelingToolset/Source/ModelingComponents/Private/Scene/WorldRenderCapture.cpp
--- a/Engine/Plugins/Runtime/MeshModelingToolset/Source/ModelingComponents/Private/Scene/WorldRenderCapture.cpp
+++ b/Engine/Plugins/Runtime/MeshModelingToolset/Source/ModelingComponents/Private/Scene/WorldRenderCapture.cpp
@@ -381,7 +381,7 @@ bool FWorldRenderCapture::CaptureMRSFromPosition(
 
 	FSceneViewFamilyContext ViewFamily(FSceneViewFamily::ConstructionValues(
 		RenderTargetResource, Scene, ShowFlags)
-		.SetWorldTimes(0, 0, 0)
+		.SetTime(FGameTime())
 		.SetRealtimeUpdate(false));
 
 	// unclear whether these show flags are really necessary, since we are using
@@ -523,7 +523,7 @@ bool FWorldRenderCapture::CaptureEmissiveFromPosition(
 
 	FSceneViewFamilyContext ViewFamily(
 		FSceneViewFamily::ConstructionValues(RenderTargetResource, World->Scene, ShowFlags)
-		.SetWorldTimes(FApp::GetCurrentTime() - GStartTime, FApp::GetDeltaTime(), FApp::GetCurrentTime() - GStartTime)
+		.SetTime(FGameTime::GetTimeSinceAppStart())
 	);
 
 	// To enable visualization mode
@@ -571,7 +571,7 @@ bool FWorldRenderCapture::CaptureEmissiveFromPosition(
 	ViewFamily.EngineShowFlags.SetCapsuleShadows(false);
 	ViewFamily.EngineShowFlags.SetContactShadows(false);
 
-	FCanvas Canvas(RenderTargetResource, NULL, FApp::GetCurrentTime() - GStartTime, FApp::GetDeltaTime(), FApp::GetCurrentTime() - GStartTime, World->Scene->GetFeatureLevel());
+	FCanvas Canvas(RenderTargetResource, NULL, FGameTime::GetTimeSinceAppStart(), World->Scene->GetFeatureLevel());
 	Canvas.Clear(FLinearColor::Transparent);
 
 	UE::Internal::PerformSceneRender(&Canvas, &ViewFamily);
@@ -625,7 +625,7 @@ namespace Internal
 
 		FSceneViewFamilyContext ViewFamily(
 			FSceneViewFamily::ConstructionValues(RenderTargetResource, Scene, FEngineShowFlags(ESFIM_Game))
-			.SetWorldTimes(FApp::GetCurrentTime() - GStartTime, FApp::GetDeltaTime(), FApp::GetCurrentTime() - GStartTime)
+			.SetTime(FGameTime::GetTimeSinceAppStart())
 		);
 
 		// To enable visualization mode
@@ -653,7 +653,7 @@ namespace Internal
 		ViewFamily.SetScreenPercentageInterface(new FLegacyScreenPercentageDriver(ViewFamily, 1.0f));
 
 		// should we cache the FCanvas?
-		FCanvas Canvas(RenderTargetResource, NULL, FApp::GetCurrentTime() - GStartTime, FApp::GetDeltaTime(), FApp::GetCurrentTime() - GStartTime, Scene->GetFeatureLevel());
+		FCanvas Canvas(RenderTargetResource, NULL, FGameTime::GetTimeSinceAppStart(), Scene->GetFeatureLevel());
 		Canvas.Clear(FLinearColor::Transparent);
 
 		UE::Internal::PerformSceneRender(&Canvas, &ViewFamily);
diff --git a/Engine/Plugins/Runtime/nDisplay/Source/DisplayCluster/Private/Game/EngineClasses/Scene/DisplayClusterPreviewComponent.cpp b/Engine/Plugins/Runtime/nDisplay/Source/DisplayCluster/Private/Game/EngineClasses/Scene/DisplayClusterPreviewComponent.cpp
--- a/Engine/Plugins/Runtime/nDisplay/Source/DisplayCluster/Private/Game/EngineClasses/Scene/DisplayClusterPreviewComponent.cpp
+++ b/Engine/Plugins/Runtime/nDisplay/Source/DisplayCluster/Private/Game/EngineClasses/Scene/DisplayClusterPreviewComponent.cpp
@@ -255,7 +255,7 @@ void UDisplayClusterPreviewComponent::UpdatePreviewRenderTarget()
 			FTextureRenderTarget2DResource* TexResource = (FTextureRenderTarget2DResource*)RenderTarget->GetResource();
 			if (TexResource)
 			{
-				FCanvas Canvas(TexResource, NULL, 0, 0, 0, GMaxRHIFeatureLevel);
+				FCanvas Canvas(TexResource, NULL, FGameTime(), GMaxRHIFeatureLevel);
 				Canvas.Clear(FLinearColor::Black);
 			}
 		}
diff --git a/Engine/Source/Developer/FunctionalTesting/Private/AutomationBlueprintFunctionLibrary.cpp b/Engine/Source/Developer/FunctionalTesting/Private/AutomationBlueprintFunctionLibrary.cpp
--- a/Engine/Source/Developer/FunctionalTesting/Private/AutomationBlueprintFunctionLibrary.cpp
+++ b/Engine/Source/Developer/FunctionalTesting/Private/AutomationBlueprintFunctionLibrary.cpp
@@ -191,11 +191,9 @@ public:
 		}
 
 		if (Options.bOverride_OverrideTimeTo)
-{
+		{
 			// Turn off time the ultimate source of noise.
-			InViewFamily.CurrentWorldTime = Options.OverrideTimeTo;
-			InViewFamily.CurrentRealTime = Options.OverrideTimeTo;
-			InViewFamily.DeltaWorldTime = 0;
+			InViewFamily.Time = FGameTime::CreateUndilated(Options.OverrideTimeTo, 0.0f);
 		}
 
 		if (Options.bDisableNoisyRenderingFeatures)
diff --git a/Engine/Source/Developer/MaterialBaking/Private/MaterialBakingModule.cpp b/Engine/Source/Developer/MaterialBaking/Private/MaterialBakingModule.cpp
--- a/Engine/Source/Developer/MaterialBaking/Private/MaterialBakingModule.cpp
+++ b/Engine/Source/Developer/MaterialBaking/Private/MaterialBakingModule.cpp
@@ -587,14 +587,14 @@ void FMaterialBakingModule::BakeMaterials(const TArray<FMaterialDataEx*>& Materi
 
 							FSceneViewFamily ViewFamily(FSceneViewFamily::ConstructionValues(RenderTarget->GetRenderTargetResource(), nullptr,
 								FEngineShowFlags(ESFIM_Game))
-								.SetWorldTimes(0.0f, 0.0f, 0.0f)
+								.SetTime(FGameTime())
 								.SetGammaCorrection(RenderTarget->GetRenderTargetResource()->GetDisplayGamma()));
 
 							RenderItem.MaterialRenderProxy = ExportMaterialProxy;
 							RenderItem.ViewFamily = &ViewFamily;
 
 							FTextureRenderTargetResource* RenderTargetResource = RenderTarget->GetRenderTargetResource();
-							FCanvas Canvas(RenderTargetResource, nullptr, FApp::GetCurrentTime() - GStartTime, FApp::GetDeltaTime(), FApp::GetCurrentTime() - GStartTime, GMaxRHIFeatureLevel);
+							FCanvas Canvas(RenderTargetResource, nullptr, FGameTime::GetTimeSinceAppStart(), GMaxRHIFeatureLevel);
 							Canvas.SetAllowedModes(FCanvas::Allow_Flush);
 							Canvas.SetRenderTargetRect(FIntRect(0, 0, RenderTarget->GetSurfaceWidth(), RenderTarget->GetSurfaceHeight()));
 							Canvas.SetBaseTransform(Canvas.CalcBaseTransform2D(RenderTarget->GetSurfaceWidth(), RenderTarget->GetSurfaceHeight()));
diff --git a/Engine/Source/Developer/MaterialUtilities/Private/MaterialUtilities.cpp b/Engine/Source/Developer/MaterialUtilities/Private/MaterialUtilities.cpp
--- a/Engine/Source/Developer/MaterialUtilities/Private/MaterialUtilities.cpp
+++ b/Engine/Source/Developer/MaterialUtilities/Private/MaterialUtilities.cpp
@@ -750,7 +750,7 @@ static void RenderSceneToTexture(
 
 	FSceneViewFamilyContext ViewFamily(
 		FSceneViewFamily::ConstructionValues(RenderTargetResource, Scene, FEngineShowFlags(ESFIM_Game))
-			.SetWorldTimes(FApp::GetCurrentTime() - GStartTime, FApp::GetDeltaTime(), FApp::GetCurrentTime() - GStartTime)
+			.SetTime(FGameTime::GetTimeSinceAppStart())
 		);
 
 	// To enable visualization mode
@@ -774,7 +774,7 @@ static void RenderSceneToTexture(
 	ViewFamily.SetScreenPercentageInterface(new FLegacyScreenPercentageDriver(
 		ViewFamily, /* GlobalResolutionFraction = */ 1.0f));
 
-	FCanvas Canvas(RenderTargetResource, NULL, FApp::GetCurrentTime() - GStartTime, FApp::GetDeltaTime(), FApp::GetCurrentTime() - GStartTime, Scene->GetFeatureLevel());
+	FCanvas Canvas(RenderTargetResource, NULL, FGameTime::GetTimeSinceAppStart(), Scene->GetFeatureLevel());
 	Canvas.Clear(FLinearColor::Transparent);
 	GetRendererModule().BeginRenderingViewFamily(&Canvas, &ViewFamily);
 
diff --git a/Engine/Source/Developer/MaterialUtilities/Private/MeshRendering.cpp b/Engine/Source/Developer/MaterialUtilities/Private/MeshRendering.cpp
--- a/Engine/Source/Developer/MaterialUtilities/Private/MeshRendering.cpp
+++ b/Engine/Source/Developer/MaterialUtilities/Private/MeshRendering.cpp
@@ -624,7 +624,7 @@ bool FMeshRenderer::RenderMaterial(struct FMaterialMergeData& InMaterialData, FM
 
 	{
 		// Create a canvas for the render target and clear it to black
-		FCanvas Canvas(RTResource, NULL, FApp::GetCurrentTime() - GStartTime, FApp::GetDeltaTime(), FApp::GetCurrentTime() - GStartTime, GMaxRHIFeatureLevel);
+		FCanvas Canvas(RTResource, NULL, FGameTime::GetTimeSinceAppStart(), GMaxRHIFeatureLevel);
 
 #if 0	// original FFlattenMaterial code - kept here for comparison
 
@@ -644,16 +644,12 @@ bool FMeshRenderer::RenderMaterial(struct FMaterialMergeData& InMaterialData, FM
 #else
 
 		// create ViewFamily
-		float CurrentRealTime = 0.f;
-		float CurrentWorldTime = 0.f;
-		float DeltaWorldTime = 0.f;
-
 		const FRenderTarget* CanvasRenderTarget = Canvas.GetRenderTarget();
 		FSceneViewFamily ViewFamily(FSceneViewFamily::ConstructionValues(
 			CanvasRenderTarget,
 			NULL,
 			FEngineShowFlags(ESFIM_Game))
-			.SetWorldTimes(CurrentWorldTime, DeltaWorldTime, CurrentRealTime)
+			.SetTime(FGameTime())
 			.SetGammaCorrection(CanvasRenderTarget->GetDisplayGamma()));
 		
 #if !SHOW_WIREFRAME_MESH
@@ -759,14 +755,10 @@ bool FMeshRenderer::RenderMaterialTexCoordScales(struct FMaterialMergeData& InMa
 {
 	check(IsInGameThread());
 	check(InRenderTarget);
-	// create ViewFamily
-	float CurrentRealTime = 0.f;
-	float CurrentWorldTime = 0.f;
-	float DeltaWorldTime = 0.f;
 
 	// Create a canvas for the render target and clear it to black
 	FTextureRenderTargetResource* RTResource = InRenderTarget->GameThread_GetRenderTargetResource();
-	FCanvas Canvas(RTResource, NULL, FApp::GetCurrentTime() - GStartTime, FApp::GetDeltaTime(), FApp::GetCurrentTime() - GStartTime, GMaxRHIFeatureLevel);
+	FCanvas Canvas(RTResource, NULL, FGameTime::GetTimeSinceAppStart(), GMaxRHIFeatureLevel);
 	const FRenderTarget* CanvasRenderTarget = Canvas.GetRenderTarget();
 	Canvas.Clear(FLinearColor::Black);
 
@@ -776,7 +768,7 @@ bool FMeshRenderer::RenderMaterialTexCoordScales(struct FMaterialMergeData& InMa
 	ShowFlags.OutputMaterialTextureScales = true; // This will bind the DVSM_OutputMaterialTextureScales
 
 	FSceneViewFamily ViewFamily(FSceneViewFamily::ConstructionValues(CanvasRenderTarget, nullptr, ShowFlags)
-		.SetWorldTimes(CurrentWorldTime, DeltaWorldTime, CurrentRealTime)
+		.SetTime(FGameTime())
 		.SetGammaCorrection(CanvasRenderTarget->GetDisplayGamma()));
 
 	// The next line ensures a constant view vector of (0,0,1) for all pixels. Required because here SVPositionToTranslatedWorld is identity, making excessive view angle increase per pixel.
diff --git a/Engine/Source/Editor/FoliageEdit/Private/FoliageType_ISMThumbnailRenderer.cpp b/Engine/Source/Editor/FoliageEdit/Private/FoliageType_ISMThumbnailRenderer.cpp
--- a/Engine/Source/Editor/FoliageEdit/Private/FoliageType_ISMThumbnailRenderer.cpp
+++ b/Engine/Source/Editor/FoliageEdit/Private/FoliageType_ISMThumbnailRenderer.cpp
@@ -29,7 +29,7 @@ void UFoliageType_ISMThumbnailRenderer::Draw(UObject* Object, int32 X, int32 Y,
 		ThumbnailScene->GetScene()->UpdateSpeedTreeWind(0.0);
 
 		FSceneViewFamilyContext ViewFamily(FSceneViewFamily::ConstructionValues(RenderTarget, ThumbnailScene->GetScene(), FEngineShowFlags(ESFIM_Game))
-			.SetWorldTimes(FApp::GetCurrentTime() - GStartTime, FApp::GetDeltaTime(), FApp::GetCurrentTime() - GStartTime)
+			.SetTime(UThumbnailRenderer::GetTime())
 			.SetAdditionalViewFamily(bAdditionalViewFamily));
 
 		ViewFamily.EngineShowFlags.DisableAdvancedFeatures();
diff --git a/Engine/Source/Editor/GraphEditor/Private/MaterialNodes/SGraphNodeMaterialBase.cpp b/Engine/Source/Editor/GraphEditor/Private/MaterialNodes/SGraphNodeMaterialBase.cpp
--- a/Engine/Source/Editor/GraphEditor/Private/MaterialNodes/SGraphNodeMaterialBase.cpp
+++ b/Engine/Source/Editor/GraphEditor/Private/MaterialNodes/SGraphNodeMaterialBase.cpp
@@ -224,7 +224,7 @@ void FPreviewElement::DrawRenderThread(FRHICommandListImmediate& RHICmdList, con
 			float CurrentTime = bIsRealtime ? (FApp::GetCurrentTime() - GStartTime) : 0.0f;
 			float DeltaTime = bIsRealtime ? FApp::GetDeltaTime() : 0.0f;
 
-			FCanvas Canvas(RenderTarget, NULL, CurrentTime, CurrentTime, DeltaTime, GMaxRHIFeatureLevel);
+			FCanvas Canvas(RenderTarget, NULL, FGameTime::CreateUndilated(CurrentTime, DeltaTime), GMaxRHIFeatureLevel);
 			{
 				Canvas.SetAllowedModes(0);
 				Canvas.SetRenderTargetRect(RenderTarget->GetViewRect());
diff --git a/Engine/Source/Editor/MeshPaint/Private/TexturePaintHelpers.cpp b/Engine/Source/Editor/MeshPaint/Private/TexturePaintHelpers.cpp
--- a/Engine/Source/Editor/MeshPaint/Private/TexturePaintHelpers.cpp
+++ b/Engine/Source/Editor/MeshPaint/Private/TexturePaintHelpers.cpp
@@ -28,7 +28,7 @@ void TexturePaintHelpers::CopyTextureToRenderTargetTexture(UTexture* SourceTextu
 	check(RenderTargetResource != nullptr);
 	
 	// Create a canvas for the render target and clear it to black
-	FCanvas Canvas(RenderTargetResource, nullptr, 0, 0, 0, FeatureLevel);
+	FCanvas Canvas(RenderTargetResource, nullptr, FGameTime(), FeatureLevel);
 
 	const uint32 Width = RenderTargetTexture->GetSurfaceWidth();
 	const uint32 Height = RenderTargetTexture->GetSurfaceHeight();
@@ -204,7 +204,7 @@ bool TexturePaintHelpers::GenerateSeamMask(UMeshComponent* MeshComponent, int32
 
 	{
 		// Create a canvas for the render target and clear it to white
-		FCanvas Canvas(RenderTargetResource, nullptr, 0, 0, 0, GEditor->GetEditorWorldContext().World()->FeatureLevel);
+		FCanvas Canvas(RenderTargetResource, nullptr, FGameTime(), GEditor->GetEditorWorldContext().World()->FeatureLevel);
 		Canvas.Clear(FLinearColor::White);
 
 		TArray<FCanvasUVTri> TriList;
diff --git a/Engine/Source/Editor/MovieSceneTools/Private/TrackEditorThumbnail/TrackEditorThumbnail.cpp b/Engine/Source/Editor/MovieSceneTools/Private/TrackEditorThumbnail/TrackEditorThumbnail.cpp
--- a/Engine/Source/Editor/MovieSceneTools/Private/TrackEditorThumbnail/TrackEditorThumbnail.cpp
+++ b/Engine/Source/Editor/MovieSceneTools/Private/TrackEditorThumbnail/TrackEditorThumbnail.cpp
@@ -41,8 +41,7 @@ public:
 		FSceneView* View = FLevelEditorViewportClient::CalcSceneView(ViewFamily, StereoViewIndex);
 
 		// Artificially set the world times so that graphics settings apply correctly (we don't tick the world when rendering thumbnails)
-		ViewFamily->CurrentWorldTime = CurrentWorldTime;
-		ViewFamily->DeltaWorldTime = DeltaWorldTime;
+		ViewFamily->Time = FGameTime::CreateDilated(ViewFamily->Time.GetRealTimeSeconds(), ViewFamily->Time.GetDeltaRealTimeSeconds(), CurrentWorldTime, DeltaWorldTime);
 
 		View->FinalPostProcessSettings.bOverride_AutoExposureSpeedDown = View->FinalPostProcessSettings.bOverride_AutoExposureSpeedUp = true;
 		View->FinalPostProcessSettings.AutoExposureSpeedDown = View->FinalPostProcessSettings.AutoExposureSpeedUp = 0.02f;
@@ -404,7 +403,7 @@ void FTrackEditorThumbnailCache::DrawViewportThumbnail(FTrackEditorThumbnail& Tr
 	UWorld* World = PreviewCameraComponent->GetWorld();
 
 	FSceneViewFamilyContext ViewFamily( FSceneViewFamily::ConstructionValues( TrackEditorThumbnail.GetRenderTarget(), World->Scene, FEngineShowFlags(ESFIM_Game) )
-		.SetWorldTimes(FApp::GetCurrentTime() - GStartTime, FApp::GetDeltaTime(), FApp::GetCurrentTime() - GStartTime)
+		.SetTime(FGameTime::GetTimeSinceAppStart())
 		.SetResolveScene(true));
 
 	FSceneViewStateInterface* ViewStateInterface = nullptr;
@@ -460,7 +459,7 @@ void FTrackEditorThumbnailCache::DrawViewportThumbnail(FTrackEditorThumbnail& Tr
 	const float GlobalResolutionFraction = 1.f;
 	ViewFamily.SetScreenPercentageInterface(new FLegacyScreenPercentageDriver(ViewFamily, GlobalResolutionFraction));
 
-	FCanvas Canvas(TrackEditorThumbnail.GetRenderTarget(), nullptr, FApp::GetCurrentTime() - GStartTime, FApp::GetDeltaTime(), FApp::GetCurrentTime() - GStartTime, World->Scene->GetFeatureLevel());
+	FCanvas Canvas(TrackEditorThumbnail.GetRenderTarget(), nullptr, FGameTime::GetTimeSinceAppStart(), World->Scene->GetFeatureLevel());
 	Canvas.Clear(FLinearColor::Transparent);
 
 	GetRendererModule().BeginRenderingViewFamily(&Canvas, &ViewFamily);
diff --git a/Engine/Source/Editor/UnrealEd/Classes/ThumbnailRendering/ThumbnailRenderer.h b/Engine/Source/Editor/UnrealEd/Classes/ThumbnailRendering/ThumbnailRenderer.h
--- a/Engine/Source/Editor/UnrealEd/Classes/ThumbnailRendering/ThumbnailRenderer.h
+++ b/Engine/Source/Editor/UnrealEd/Classes/ThumbnailRendering/ThumbnailRenderer.h
@@ -84,5 +84,7 @@ public:
 protected:
 	/** Renders the thumbnail's view family. */
 	UNREALED_API static void RenderViewFamily(FCanvas* Canvas, class FSceneViewFamily* ViewFamily);
+
+	UNREALED_API static struct FGameTime GetTime();
 };
 
diff --git a/Engine/Source/Editor/UnrealEd/Private/EditorViewportClient.cpp b/Engine/Source/Editor/UnrealEd/Private/EditorViewportClient.cpp
--- a/Engine/Source/Editor/UnrealEd/Private/EditorViewportClient.cpp
+++ b/Engine/Source/Editor/UnrealEd/Private/EditorViewportClient.cpp
@@ -3750,24 +3750,15 @@ void FEditorViewportClient::Draw(FViewport* InViewport, FCanvas* Canvas)
 	FViewport* ViewportBackup = Viewport;
 	Viewport = InViewport ? InViewport : Viewport;
 
-	// Determine whether we should use world time or real time based on the scene.
-	float TimeSeconds;
-	float RealTimeSeconds;
-	float DeltaTimeSeconds;
-
 	UWorld* World = GetWorld();
+	FGameTime Time;
 	if (!World || (GetScene() != World->Scene) || UseAppTime()) 
 	{
-		// Use time relative to start time to avoid issues with float vs double
-		TimeSeconds = FApp::GetCurrentTime() - GStartTime;
-		RealTimeSeconds = FApp::GetCurrentTime() - GStartTime;
-		DeltaTimeSeconds = FApp::GetDeltaTime();
+		Time = FGameTime::GetTimeSinceAppStart();
 	}
 	else
 	{
-		TimeSeconds = World->GetTimeSeconds();
-		RealTimeSeconds = World->GetRealTimeSeconds();
-		DeltaTimeSeconds = World->GetDeltaSeconds();
+		Time = World->GetTime();
 	}
 
 	// Allow HMD to modify the view later, just before rendering
@@ -3792,7 +3783,7 @@ void FEditorViewportClient::Draw(FViewport* InViewport, FCanvas* Canvas)
 		Canvas->GetRenderTarget(),
 		GetScene(),
 		UseEngineShowFlags)
-		.SetWorldTimes( TimeSeconds, DeltaTimeSeconds, RealTimeSeconds )
+		.SetTime(Time)
 		.SetRealtimeUpdate( IsRealtime() && FSlateThrottleManager::Get().IsAllowingExpensiveTasks() )
 		.SetViewModeParam( ViewModeParam, ViewModeParamName ) );
 
diff --git a/Engine/Source/Editor/UnrealEd/Private/Lightmass/LightmassLandscapeRender.cpp b/Engine/Source/Editor/UnrealEd/Private/Lightmass/LightmassLandscapeRender.cpp
--- a/Engine/Source/Editor/UnrealEd/Private/Lightmass/LightmassLandscapeRender.cpp
+++ b/Engine/Source/Editor/UnrealEd/Private/Lightmass/LightmassLandscapeRender.cpp
@@ -91,7 +91,7 @@ void RenderLandscapeMaterialForLightmass(const FLandscapeStaticLightingMesh* Lan
 		RenderTarget,
 		NULL,
 		FEngineShowFlags(ESFIM_Game))
-		.SetWorldTimes(0, 0, 0)
+		.SetTime(FGameTime())
 		.SetGammaCorrection(RenderTarget->GetDisplayGamma()));
 
 	FDynamicMeshBuilder DynamicMeshBuilder(ViewFamily.GetFeatureLevel(), 4, 0, true);
diff --git a/Engine/Source/Editor/UnrealEd/Private/Lightmass/LightmassRender.cpp b/Engine/Source/Editor/UnrealEd/Private/Lightmass/LightmassRender.cpp
--- a/Engine/Source/Editor/UnrealEd/Private/Lightmass/LightmassRender.cpp
+++ b/Engine/Source/Editor/UnrealEd/Private/Lightmass/LightmassRender.cpp
@@ -1230,7 +1230,7 @@ bool FLightmassMaterialRenderer::CreateRenderTarget(EPixelFormat InFormat, int32
 		RenderTarget->ClearColor = FLinearColor(0.0f, 0.0f, 0.0f, 0.0f);
 		RenderTarget->InitCustomFormat(InSizeX, InSizeY, InFormat, false);
 
-		Canvas = new FCanvas(RenderTarget->GameThread_GetRenderTargetResource(), NULL, 0, 0, 0, GMaxRHIFeatureLevel);
+		Canvas = new FCanvas(RenderTarget->GameThread_GetRenderTargetResource(), NULL, FGameTime(), GMaxRHIFeatureLevel);
 		check(Canvas);
 	}
 
diff --git a/Engine/Source/Editor/UnrealEd/Private/ObjectTools.cpp b/Engine/Source/Editor/UnrealEd/Private/ObjectTools.cpp
--- a/Engine/Source/Editor/UnrealEd/Private/ObjectTools.cpp
+++ b/Engine/Source/Editor/UnrealEd/Private/ObjectTools.cpp
@@ -4514,7 +4514,7 @@ namespace ThumbnailTools
 		check( RenderTargetResource != NULL );
 
 		// Create a canvas for the render target and clear it to black
-		FCanvas Canvas( RenderTargetResource, NULL, FApp::GetCurrentTime() - GStartTime, FApp::GetDeltaTime(), FApp::GetCurrentTime() - GStartTime, GMaxRHIFeatureLevel );
+		FCanvas Canvas( RenderTargetResource, NULL, FGameTime::GetTimeSinceAppStart(), GMaxRHIFeatureLevel );
 		Canvas.Clear( FLinearColor::Black );
 
 		// Get the rendering info for this object
diff --git a/Engine/Source/Editor/UnrealEd/Private/ThumbnailRendering/AnimBlueprintThumbnailRenderer.cpp b/Engine/Source/Editor/UnrealEd/Private/ThumbnailRendering/AnimBlueprintThumbnailRenderer.cpp
--- a/Engine/Source/Editor/UnrealEd/Private/ThumbnailRendering/AnimBlueprintThumbnailRenderer.cpp
+++ b/Engine/Source/Editor/UnrealEd/Private/ThumbnailRendering/AnimBlueprintThumbnailRenderer.cpp
@@ -27,7 +27,7 @@ void UAnimBlueprintThumbnailRenderer::Draw(UObject* Object, int32 X, int32 Y, ui
 		if(ThumbnailScene->SetAnimBlueprint(AnimBlueprint))
 		{
 			FSceneViewFamilyContext ViewFamily(FSceneViewFamily::ConstructionValues(RenderTarget, ThumbnailScene->GetScene(), FEngineShowFlags(ESFIM_Game))
-				.SetWorldTimes(FApp::GetCurrentTime() - GStartTime, FApp::GetDeltaTime(), FApp::GetCurrentTime() - GStartTime)
+				.SetTime(UThumbnailRenderer::GetTime())
 				.SetAdditionalViewFamily(bAdditionalViewFamily));
 
 			ViewFamily.EngineShowFlags.DisableAdvancedFeatures();
diff --git a/Engine/Source/Editor/UnrealEd/Private/ThumbnailRendering/AnimSequenceThumbnailRenderer.cpp b/Engine/Source/Editor/UnrealEd/Private/ThumbnailRendering/AnimSequenceThumbnailRenderer.cpp
--- a/Engine/Source/Editor/UnrealEd/Private/ThumbnailRendering/AnimSequenceThumbnailRenderer.cpp
+++ b/Engine/Source/Editor/UnrealEd/Private/ThumbnailRendering/AnimSequenceThumbnailRenderer.cpp
@@ -26,7 +26,7 @@ void UAnimSequenceThumbnailRenderer::Draw(UObject* Object, int32 X, int32 Y, uin
 		if(ThumbnailScene->SetAnimation(Anim))
 		{
 			FSceneViewFamilyContext ViewFamily(FSceneViewFamily::ConstructionValues(RenderTarget, ThumbnailScene->GetScene(), FEngineShowFlags(ESFIM_Game))
-				.SetWorldTimes(FApp::GetCurrentTime() - GStartTime, FApp::GetDeltaTime(), FApp::GetCurrentTime() - GStartTime)
+				.SetTime(UThumbnailRenderer::GetTime())
 				.SetAdditionalViewFamily(bAdditionalViewFamily));
 
 			ViewFamily.EngineShowFlags.DisableAdvancedFeatures();
diff --git a/Engine/Source/Editor/UnrealEd/Private/ThumbnailRendering/BlendSpaceThumbnailRenderer.cpp b/Engine/Source/Editor/UnrealEd/Private/ThumbnailRendering/BlendSpaceThumbnailRenderer.cpp
--- a/Engine/Source/Editor/UnrealEd/Private/ThumbnailRendering/BlendSpaceThumbnailRenderer.cpp
+++ b/Engine/Source/Editor/UnrealEd/Private/ThumbnailRendering/BlendSpaceThumbnailRenderer.cpp
@@ -27,7 +27,7 @@ void UBlendSpaceThumbnailRenderer::Draw(UObject* Object, int32 X, int32 Y, uint3
 		if (ThumbnailScene->SetBlendSpace(BlendSpace))
 		{
 			FSceneViewFamilyContext ViewFamily(FSceneViewFamily::ConstructionValues(RenderTarget, ThumbnailScene->GetScene(), FEngineShowFlags(ESFIM_Game))
-				.SetWorldTimes(FApp::GetCurrentTime() - GStartTime, FApp::GetDeltaTime(), FApp::GetCurrentTime() - GStartTime)
+				.SetTime(UThumbnailRenderer::GetTime())
 				.SetAdditionalViewFamily(bAdditionalViewFamily));
 
 			ViewFamily.EngineShowFlags.DisableAdvancedFeatures();
diff --git a/Engine/Source/Editor/UnrealEd/Private/ThumbnailRendering/BlueprintThumbnailRenderer.cpp b/Engine/Source/Editor/UnrealEd/Private/ThumbnailRendering/BlueprintThumbnailRenderer.cpp
--- a/Engine/Source/Editor/UnrealEd/Private/ThumbnailRendering/BlueprintThumbnailRenderer.cpp
+++ b/Engine/Source/Editor/UnrealEd/Private/ThumbnailRendering/BlueprintThumbnailRenderer.cpp
@@ -105,7 +105,7 @@ void UBlueprintThumbnailRenderer::Draw(UObject* Object, int32 X, int32 Y, uint32
 
 		ThumbnailScene->SetBlueprint(Blueprint);
 		FSceneViewFamilyContext ViewFamily( FSceneViewFamily::ConstructionValues( RenderTarget, ThumbnailScene->GetScene(), FEngineShowFlags(ESFIM_Game) )
-			.SetWorldTimes(FApp::GetCurrentTime() - GStartTime, FApp::GetDeltaTime(), FApp::GetCurrentTime() - GStartTime)
+			.SetTime(UThumbnailRenderer::GetTime())
 			.SetAdditionalViewFamily(bAdditionalViewFamily));
 
 		ViewFamily.EngineShowFlags.DisableAdvancedFeatures();
diff --git a/Engine/Source/Editor/UnrealEd/Private/ThumbnailRendering/ClassThumbnailRenderer.cpp b/Engine/Source/Editor/UnrealEd/Private/ThumbnailRendering/ClassThumbnailRenderer.cpp
--- a/Engine/Source/Editor/UnrealEd/Private/ThumbnailRendering/ClassThumbnailRenderer.cpp
+++ b/Engine/Source/Editor/UnrealEd/Private/ThumbnailRendering/ClassThumbnailRenderer.cpp
@@ -41,7 +41,7 @@ void UClassThumbnailRenderer::Draw(UObject* Object, int32 X, int32 Y, uint32 Wid
 
 		ThumbnailScene->SetClass(Class);
 		FSceneViewFamilyContext ViewFamily( FSceneViewFamily::ConstructionValues( RenderTarget, ThumbnailScene->GetScene(), FEngineShowFlags(ESFIM_Game) )
-			.SetWorldTimes(FApp::GetCurrentTime() - GStartTime, FApp::GetDeltaTime(), FApp::GetCurrentTime() - GStartTime)
+			.SetTime(UThumbnailRenderer::GetTime())
 			.SetAdditionalViewFamily(bAdditionalViewFamily));
 
 		ViewFamily.EngineShowFlags.DisableAdvancedFeatures();
diff --git a/Engine/Source/Editor/UnrealEd/Private/ThumbnailRendering/LevelThumbnailRenderer.cpp b/Engine/Source/Editor/UnrealEd/Private/ThumbnailRendering/LevelThumbnailRenderer.cpp
--- a/Engine/Source/Editor/UnrealEd/Private/ThumbnailRendering/LevelThumbnailRenderer.cpp
+++ b/Engine/Source/Editor/UnrealEd/Private/ThumbnailRendering/LevelThumbnailRenderer.cpp
@@ -19,7 +19,7 @@ void ULevelThumbnailRenderer::Draw(UObject* Object, int32 X, int32 Y, uint32 Wid
 	if (Level != nullptr)
 	{
 		FSceneViewFamilyContext ViewFamily( FSceneViewFamily::ConstructionValues( RenderTarget, Level->OwningWorld->Scene, FEngineShowFlags(ESFIM_Game) )
-			.SetWorldTimes(FApp::GetCurrentTime() - GStartTime, FApp::GetDeltaTime(), FApp::GetCurrentTime() - GStartTime)
+			.SetTime(UThumbnailRenderer::GetTime())
 			.SetAdditionalViewFamily(bAdditionalViewFamily));
 
 		ViewFamily.EngineShowFlags.DisableAdvancedFeatures();
diff --git a/Engine/Source/Editor/UnrealEd/Private/ThumbnailRendering/MaterialFunctionThumbnailRenderer.cpp b/Engine/Source/Editor/UnrealEd/Private/ThumbnailRendering/MaterialFunctionThumbnailRenderer.cpp
--- a/Engine/Source/Editor/UnrealEd/Private/ThumbnailRendering/MaterialFunctionThumbnailRenderer.cpp
+++ b/Engine/Source/Editor/UnrealEd/Private/ThumbnailRendering/MaterialFunctionThumbnailRenderer.cpp
@@ -51,7 +51,7 @@ void UMaterialFunctionThumbnailRenderer::Draw(UObject* Object, int32 X, int32 Y,
 				ThumbnailScene->SetMaterialInterface(PreviewMaterial);
 	
 			FSceneViewFamilyContext ViewFamily( FSceneViewFamily::ConstructionValues( RenderTarget, ThumbnailScene->GetScene(), FEngineShowFlags(ESFIM_Game) )
-				.SetWorldTimes(FApp::GetCurrentTime() - GStartTime, FApp::GetDeltaTime(), FApp::GetCurrentTime() - GStartTime)
+				.SetTime(UThumbnailRenderer::GetTime())
 				.SetAdditionalViewFamily(bAdditionalViewFamily));
 
 			ViewFamily.EngineShowFlags.DisableAdvancedFeatures();
diff --git a/Engine/Source/Editor/UnrealEd/Private/ThumbnailRendering/MaterialInstanceThumbnailRenderer.cpp b/Engine/Source/Editor/UnrealEd/Private/ThumbnailRendering/MaterialInstanceThumbnailRenderer.cpp
--- a/Engine/Source/Editor/UnrealEd/Private/ThumbnailRendering/MaterialInstanceThumbnailRenderer.cpp
+++ b/Engine/Source/Editor/UnrealEd/Private/ThumbnailRendering/MaterialInstanceThumbnailRenderer.cpp
@@ -87,7 +87,7 @@ void UMaterialInstanceThumbnailRenderer::Draw(UObject* Object, int32 X, int32 Y,
 
 			ThumbnailScene->SetMaterialInterface(MatInst);
 			FSceneViewFamilyContext ViewFamily( FSceneViewFamily::ConstructionValues( RenderTarget, ThumbnailScene->GetScene(), FEngineShowFlags(ESFIM_Game) )
-				.SetWorldTimes(FApp::GetCurrentTime() - GStartTime, FApp::GetDeltaTime(), FApp::GetCurrentTime() - GStartTime)
+				.SetTime(UThumbnailRenderer::GetTime())
 				.SetAdditionalViewFamily(bAdditionalViewFamily));
 
 
diff --git a/Engine/Source/Editor/UnrealEd/Private/ThumbnailRendering/ParticleSystemThumbnailRenderer.cpp b/Engine/Source/Editor/UnrealEd/Private/ThumbnailRendering/ParticleSystemThumbnailRenderer.cpp
--- a/Engine/Source/Editor/UnrealEd/Private/ThumbnailRendering/ParticleSystemThumbnailRenderer.cpp
+++ b/Engine/Source/Editor/UnrealEd/Private/ThumbnailRendering/ParticleSystemThumbnailRenderer.cpp
@@ -77,7 +77,7 @@ void UParticleSystemThumbnailRenderer::Draw(UObject* Object, int32 X, int32 Y, u
 
 				ThumbnailScene->SetParticleSystem(ParticleSystem);
 				FSceneViewFamilyContext ViewFamily( FSceneViewFamily::ConstructionValues( RenderTarget, ThumbnailScene->GetScene(), FEngineShowFlags(ESFIM_Game))
-					.SetWorldTimes(FApp::GetCurrentTime() - GStartTime, FApp::GetDeltaTime(), FApp::GetCurrentTime() - GStartTime)
+					.SetTime(UThumbnailRenderer::GetTime())
 					.SetAdditionalViewFamily(bAdditionalViewFamily));
 			
 				ViewFamily.EngineShowFlags.DisableAdvancedFeatures();
diff --git a/Engine/Source/Editor/UnrealEd/Private/ThumbnailRendering/PhysicsAssetThumbnailRenderer.cpp b/Engine/Source/Editor/UnrealEd/Private/ThumbnailRendering/PhysicsAssetThumbnailRenderer.cpp
--- a/Engine/Source/Editor/UnrealEd/Private/ThumbnailRendering/PhysicsAssetThumbnailRenderer.cpp
+++ b/Engine/Source/Editor/UnrealEd/Private/ThumbnailRendering/PhysicsAssetThumbnailRenderer.cpp
@@ -25,7 +25,7 @@ void UPhysicsAssetThumbnailRenderer::Draw(UObject* Object, int32 X, int32 Y, uin
 
 		ThumbnailScene->SetPhysicsAsset(PhysicsAsset);
 		FSceneViewFamilyContext ViewFamily( FSceneViewFamily::ConstructionValues( RenderTarget, ThumbnailScene->GetScene(), FEngineShowFlags(ESFIM_Game) )
-			.SetWorldTimes(FApp::GetCurrentTime() - GStartTime, FApp::GetDeltaTime(), FApp::GetCurrentTime() - GStartTime)
+			.SetTime(UThumbnailRenderer::GetTime())
 			.SetAdditionalViewFamily(bAdditionalViewFamily));
 
 		ViewFamily.EngineShowFlags.DisableAdvancedFeatures();
diff --git a/Engine/Source/Editor/UnrealEd/Private/ThumbnailRendering/SkeletalMeshThumbnailRenderer.cpp b/Engine/Source/Editor/UnrealEd/Private/ThumbnailRendering/SkeletalMeshThumbnailRenderer.cpp
--- a/Engine/Source/Editor/UnrealEd/Private/ThumbnailRendering/SkeletalMeshThumbnailRenderer.cpp
+++ b/Engine/Source/Editor/UnrealEd/Private/ThumbnailRendering/SkeletalMeshThumbnailRenderer.cpp
@@ -24,7 +24,7 @@ void USkeletalMeshThumbnailRenderer::Draw(UObject* Object, int32 X, int32 Y, uin
 	AddAdditionalPreviewSceneContent(Object, ThumbnailScene->GetWorld());
 
 	FSceneViewFamilyContext ViewFamily( FSceneViewFamily::ConstructionValues( RenderTarget, ThumbnailScene->GetScene(), FEngineShowFlags(ESFIM_Game) )
-		.SetWorldTimes(FApp::GetCurrentTime() - GStartTime, FApp::GetDeltaTime(), FApp::GetCurrentTime() - GStartTime)
+		.SetTime(UThumbnailRenderer::GetTime())
 		.SetAdditionalViewFamily(bAdditionalViewFamily));
 
 	ViewFamily.EngineShowFlags.DisableAdvancedFeatures();
diff --git a/Engine/Source/Editor/UnrealEd/Private/ThumbnailRendering/StaticMeshThumbnailRenderer.cpp b/Engine/Source/Editor/UnrealEd/Private/ThumbnailRendering/StaticMeshThumbnailRenderer.cpp
--- a/Engine/Source/Editor/UnrealEd/Private/ThumbnailRendering/StaticMeshThumbnailRenderer.cpp
+++ b/Engine/Source/Editor/UnrealEd/Private/ThumbnailRendering/StaticMeshThumbnailRenderer.cpp
@@ -32,7 +32,7 @@ void UStaticMeshThumbnailRenderer::Draw(UObject* Object, int32 X, int32 Y, uint3
 		ThumbnailScene->GetScene()->UpdateSpeedTreeWind(0.0);
 
 		FSceneViewFamilyContext ViewFamily( FSceneViewFamily::ConstructionValues( RenderTarget, ThumbnailScene->GetScene(), FEngineShowFlags(ESFIM_Game) )
-			.SetWorldTimes(FApp::GetCurrentTime() - GStartTime, FApp::GetDeltaTime(), FApp::GetCurrentTime() - GStartTime)
+			.SetTime(UThumbnailRenderer::GetTime())
 			.SetAdditionalViewFamily(bAdditionalViewFamily));
 
 		ViewFamily.EngineShowFlags.DisableAdvancedFeatures();
diff --git a/Engine/Source/Editor/UnrealEd/Private/ThumbnailRendering/ThumbnailRenderer.cpp b/Engine/Source/Editor/UnrealEd/Private/ThumbnailRendering/ThumbnailRenderer.cpp
--- a/Engine/Source/Editor/UnrealEd/Private/ThumbnailRendering/ThumbnailRenderer.cpp
+++ b/Engine/Source/Editor/UnrealEd/Private/ThumbnailRendering/ThumbnailRenderer.cpp
@@ -21,3 +21,9 @@ void UThumbnailRenderer::RenderViewFamily(FCanvas* Canvas, FSceneViewFamily* Vie
 
 	GetRendererModule().BeginRenderingViewFamily(Canvas, ViewFamily);
 }
+
+// static
+FGameTime UThumbnailRenderer::GetTime()
+{
+	return FGameTime::GetTimeSinceAppStart();
+}
diff --git a/Engine/Source/Editor/UnrealEd/Private/ThumbnailRendering/VolumeTextureThumbnailRenderer.cpp b/Engine/Source/Editor/UnrealEd/Private/ThumbnailRendering/VolumeTextureThumbnailRenderer.cpp
--- a/Engine/Source/Editor/UnrealEd/Private/ThumbnailRendering/VolumeTextureThumbnailRenderer.cpp
+++ b/Engine/Source/Editor/UnrealEd/Private/ThumbnailRendering/VolumeTextureThumbnailRenderer.cpp
@@ -161,7 +161,7 @@ void UVolumeTextureThumbnailRenderer::Draw(UObject* Object, int32 X, int32 Y, ui
 
 			ThumbnailScene->SetMaterialInterface(MaterialInstance);
 			FSceneViewFamilyContext ViewFamily( FSceneViewFamily::ConstructionValues( RenderTarget, ThumbnailScene->GetScene(), FEngineShowFlags(ESFIM_Game) )
-				.SetWorldTimes(FApp::GetCurrentTime() - GStartTime, FApp::GetDeltaTime(), FApp::GetCurrentTime() - GStartTime)
+				.SetTime(UThumbnailRenderer::GetTime())
 				.SetAdditionalViewFamily(bAdditionalViewFamily));
 
 			ViewFamily.EngineShowFlags.DisableAdvancedFeatures();
diff --git a/Engine/Source/Editor/UnrealEd/Private/ThumbnailRendering/WorldThumbnailRenderer.cpp b/Engine/Source/Editor/UnrealEd/Private/ThumbnailRendering/WorldThumbnailRenderer.cpp
--- a/Engine/Source/Editor/UnrealEd/Private/ThumbnailRendering/WorldThumbnailRenderer.cpp
+++ b/Engine/Source/Editor/UnrealEd/Private/ThumbnailRendering/WorldThumbnailRenderer.cpp
@@ -40,7 +40,7 @@ void UWorldThumbnailRenderer::Draw(UObject* Object, int32 X, int32 Y, uint32 Wid
 	if (World != nullptr && World->Scene)
 	{
 		FSceneViewFamilyContext ViewFamily( FSceneViewFamily::ConstructionValues( RenderTarget, World->Scene, FEngineShowFlags(ESFIM_All0) )
-			.SetWorldTimes(FApp::GetCurrentTime() - GStartTime, FApp::GetDeltaTime(), FApp::GetCurrentTime() - GStartTime)
+			.SetTime(UThumbnailRenderer::GetTime())
 			.SetAdditionalViewFamily(bAdditionalViewFamily));
 
 		ViewFamily.EngineShowFlags.SetDiffuse(true);
diff --git a/Engine/Source/Runtime/Engine/Private/CanvasRenderTarget2D.cpp b/Engine/Source/Runtime/Engine/Private/CanvasRenderTarget2D.cpp
--- a/Engine/Source/Runtime/Engine/Private/CanvasRenderTarget2D.cpp
+++ b/Engine/Source/Runtime/Engine/Private/CanvasRenderTarget2D.cpp
@@ -71,7 +71,7 @@ void UCanvasRenderTarget2D::RepaintCanvas()
 	// NOTE: This texture may be null when this is invoked through blueprint from a cmdlet or server.
 	FTextureRenderTarget2DResource* TextureRenderTarget = (FTextureRenderTarget2DResource*) GameThread_GetRenderTargetResource();
 
-	FCanvas RenderCanvas(TextureRenderTarget, nullptr, FApp::GetCurrentTime() - GStartTime, FApp::GetDeltaTime(), FApp::GetCurrentTime() - GStartTime, FeatureLevel);
+	FCanvas RenderCanvas(TextureRenderTarget, nullptr, FGameTime::GetTimeSinceAppStart(), FeatureLevel);
 	Canvas->Init(GetSurfaceWidth(), GetSurfaceHeight(), nullptr, &RenderCanvas);
 
 	if (TextureRenderTarget)
diff --git a/Engine/Source/Runtime/Engine/Private/CubemapUnwrapUtils.cpp b/Engine/Source/Runtime/Engine/Private/CubemapUnwrapUtils.cpp
--- a/Engine/Source/Runtime/Engine/Private/CubemapUnwrapUtils.cpp
+++ b/Engine/Source/Runtime/Engine/Private/CubemapUnwrapUtils.cpp
@@ -51,7 +51,7 @@ namespace CubemapHelpers
 		RenderTargetLongLat->TargetGamma = 0;
 		FRenderTarget* RenderTarget = RenderTargetLongLat->GameThread_GetRenderTargetResource();
 
-		FCanvas* Canvas = new FCanvas(RenderTarget, NULL, 0, 0, 0, GMaxRHIFeatureLevel);
+		FCanvas* Canvas = new FCanvas(RenderTarget, NULL, FGameTime(), GMaxRHIFeatureLevel);
 		Canvas->SetRenderTarget_GameThread(RenderTarget);
 
 		// Clear the render target to black
diff --git a/Engine/Source/Runtime/Engine/Private/TileRendering.cpp b/Engine/Source/Runtime/Engine/Private/TileRendering.cpp
--- a/Engine/Source/Runtime/Engine/Private/TileRendering.cpp
+++ b/Engine/Source/Runtime/Engine/Private/TileRendering.cpp
@@ -229,15 +229,10 @@ void FCanvasTileRendererItem::FRenderData::RenderTiles(
 
 bool FCanvasTileRendererItem::Render_RenderThread(FCanvasRenderContext& RenderContext, FMeshPassProcessorRenderState& DrawRenderState, const FCanvas* Canvas)
 {
-	float CurrentRealTime = 0.f;
-	float CurrentWorldTime = 0.f;
-	float DeltaWorldTime = 0.f;
-
+	FGameTime Time;
 	if (!bFreezeTime)
 	{
-		CurrentRealTime = Canvas->GetCurrentRealTime();
-		CurrentWorldTime = Canvas->GetCurrentWorldTime();
-		DeltaWorldTime = Canvas->GetCurrentDeltaWorldTime();
+		Time = Canvas->GetTime();
 	}
 
 	checkSlow(Data);
@@ -248,7 +243,7 @@ bool FCanvasTileRendererItem::Render_RenderThread(FCanvasRenderContext& RenderCo
 		CanvasRenderTarget,
 		nullptr,
 		FEngineShowFlags(ESFIM_Game))
-		.SetWorldTimes(CurrentWorldTime, DeltaWorldTime, CurrentRealTime)
+		.SetTime(Time)
 		.SetGammaCorrection(CanvasRenderTarget->GetDisplayGamma()));
 
 	const FIntRect ViewRect(FIntPoint(0, 0), CanvasRenderTarget->GetSizeXY());
@@ -280,15 +275,10 @@ bool FCanvasTileRendererItem::Render_RenderThread(FCanvasRenderContext& RenderCo
 
 bool FCanvasTileRendererItem::Render_GameThread(const FCanvas* Canvas, FCanvasRenderThreadScope& RenderScope)
 {
-	float CurrentRealTime = 0.f;
-	float CurrentWorldTime = 0.f;
-	float DeltaWorldTime = 0.f;
-
+	FGameTime Time;
 	if (!bFreezeTime)
 	{
-		CurrentRealTime = Canvas->GetCurrentRealTime();
-		CurrentWorldTime = Canvas->GetCurrentWorldTime();
-		DeltaWorldTime = Canvas->GetCurrentDeltaWorldTime();
+		Time = Canvas->GetTime();
 	}
 
 	checkSlow(Data);
@@ -300,7 +290,7 @@ bool FCanvasTileRendererItem::Render_GameThread(const FCanvas* Canvas, FCanvasRe
 			CanvasRenderTarget,
 			Canvas->GetScene(),
 			FEngineShowFlags(ESFIM_Game))
-			.SetWorldTimes(CurrentWorldTime, DeltaWorldTime, CurrentRealTime)
+			.SetTime(Time)
 			.SetGammaCorrection(CanvasRenderTarget->GetDisplayGamma()));
 
 		const FIntRect ViewRect(FIntPoint(0, 0), CanvasRenderTarget->GetSizeXY());
diff --git a/Engine/Source/Runtime/Engine/Private/TriangleRendering.cpp b/Engine/Source/Runtime/Engine/Private/TriangleRendering.cpp
--- a/Engine/Source/Runtime/Engine/Private/TriangleRendering.cpp
+++ b/Engine/Source/Runtime/Engine/Private/TriangleRendering.cpp
@@ -200,15 +200,10 @@ void FCanvasTriangleRendererItem::FRenderData::RenderTriangles(
 
 bool FCanvasTriangleRendererItem::Render_RenderThread(FCanvasRenderContext& RenderContext, FMeshPassProcessorRenderState& DrawRenderState, const FCanvas* Canvas)
 {
-	float CurrentRealTime = 0.f;
-	float CurrentWorldTime = 0.f;
-	float DeltaWorldTime = 0.f;
-
+	FGameTime Time;
 	if (!bFreezeTime)
 	{
-		CurrentRealTime = Canvas->GetCurrentRealTime();
-		CurrentWorldTime = Canvas->GetCurrentWorldTime();
-		DeltaWorldTime = Canvas->GetCurrentDeltaWorldTime();
+		Time = Canvas->GetTime();
 	}
 
 	checkSlow(Data);
@@ -219,7 +214,7 @@ bool FCanvasTriangleRendererItem::Render_RenderThread(FCanvasRenderContext& Rend
 		CanvasRenderTarget,
 		nullptr,
 		FEngineShowFlags(ESFIM_Game))
-		.SetWorldTimes(CurrentWorldTime, DeltaWorldTime, CurrentRealTime)
+		.SetTime(Time)
 		.SetGammaCorrection(CanvasRenderTarget->GetDisplayGamma()));
 
 	const FIntRect ViewRect(FIntPoint(0, 0), CanvasRenderTarget->GetSizeXY());
@@ -251,15 +246,10 @@ bool FCanvasTriangleRendererItem::Render_RenderThread(FCanvasRenderContext& Rend
 
 bool FCanvasTriangleRendererItem::Render_GameThread(const FCanvas* Canvas, FCanvasRenderThreadScope& RenderScope)
 {
-	float CurrentRealTime = 0.f;
-	float CurrentWorldTime = 0.f;
-	float DeltaWorldTime = 0.f;
-
+	FGameTime Time;
 	if (!bFreezeTime)
 	{
-		CurrentRealTime = Canvas->GetCurrentRealTime();
-		CurrentWorldTime = Canvas->GetCurrentWorldTime();
-		DeltaWorldTime = Canvas->GetCurrentDeltaWorldTime();
+		Time = Canvas->GetTime();
 	}
 
 	checkSlow(Data);
@@ -270,7 +260,7 @@ bool FCanvasTriangleRendererItem::Render_GameThread(const FCanvas* Canvas, FCanv
 		CanvasRenderTarget,
 		Canvas->GetScene(),
 		FEngineShowFlags(ESFIM_Game))
-		.SetWorldTimes(CurrentWorldTime, DeltaWorldTime, CurrentRealTime)
+		.SetTime(Time)
 		.SetGammaCorrection(CanvasRenderTarget->GetDisplayGamma()));
 
 	const FIntRect ViewRect(FIntPoint(0, 0), CanvasRenderTarget->GetSizeXY());
diff --git a/Engine/Source/Runtime/Landscape/Private/LandscapeEdit.cpp b/Engine/Source/Runtime/Landscape/Private/LandscapeEdit.cpp
--- a/Engine/Source/Runtime/Landscape/Private/LandscapeEdit.cpp
+++ b/Engine/Source/Runtime/Landscape/Private/LandscapeEdit.cpp
@@ -7346,7 +7346,7 @@ bool ALandscapeProxy::LandscapeExportHeightmapToRenderTarget(UTextureRenderTarge
 	FTextureRenderTargetResource* RenderTargetResource = InRenderTarget->GameThread_GetRenderTargetResource();
 
 	// Create a canvas for the render target and clear it to black
-	FCanvas Canvas(RenderTargetResource, nullptr, 0, 0, 0, World->FeatureLevel);
+	FCanvas Canvas(RenderTargetResource, nullptr, FGameTime(), World->FeatureLevel);
 	Canvas.Clear(FLinearColor::Black);
 
 	// Find exported component's base offset
diff --git a/Engine/Source/Runtime/Landscape/Private/LandscapeEditLayers.cpp b/Engine/Source/Runtime/Landscape/Private/LandscapeEditLayers.cpp
--- a/Engine/Source/Runtime/Landscape/Private/LandscapeEditLayers.cpp
+++ b/Engine/Source/Runtime/Landscape/Private/LandscapeEditLayers.cpp
@@ -1287,7 +1287,8 @@ public:
 
 		check(IsInRenderingThread());
 
-		FSceneViewFamilyContext ViewFamily(FSceneViewFamily::ConstructionValues(RenderTargetResource, NULL, FEngineShowFlags(ESFIM_Game)).SetWorldTimes(FApp::GetCurrentTime() - GStartTime, FApp::GetDeltaTime(), FApp::GetCurrentTime() - GStartTime));
+		FSceneViewFamilyContext ViewFamily(FSceneViewFamily::ConstructionValues(RenderTargetResource, NULL, FEngineShowFlags(ESFIM_Game))
+			.SetTime(FGameTime::GetTimeSinceAppStart()));
 
 		FSceneViewInitOptions ViewInitOptions;
 		ViewInitOptions.SetViewRectangle(FIntRect(0, 0, WriteRenderTargetSize.X, WriteRenderTargetSize.Y));
diff --git a/Engine/Source/Runtime/Landscape/Private/LandscapeGrass.cpp b/Engine/Source/Runtime/Landscape/Private/LandscapeGrass.cpp
--- a/Engine/Source/Runtime/Landscape/Private/LandscapeGrass.cpp
+++ b/Engine/Source/Runtime/Landscape/Private/LandscapeGrass.cpp
@@ -564,7 +564,8 @@ public:
 		FMemMark Mark(FMemStack::Get());
 		FRDGBuilder GraphBuilder(RHICmdList);
 
-		FSceneViewFamilyContext ViewFamily(FSceneViewFamily::ConstructionValues(nullptr, SceneInterface, FEngineShowFlags(ESFIM_Game)).SetWorldTimes(FApp::GetCurrentTime() - GStartTime, FApp::GetDeltaTime(), FApp::GetCurrentTime() - GStartTime));
+		FSceneViewFamilyContext ViewFamily(FSceneViewFamily::ConstructionValues(nullptr, SceneInterface, FEngineShowFlags(ESFIM_Game))
+			.SetTime(FGameTime::GetTimeSinceAppStart()));
 
 		ViewFamily.LandscapeLODOverride = 0; // Force LOD render
 
diff --git a/Engine/Source/Runtime/Landscape/Private/LandscapePhysicalMaterial.cpp b/Engine/Source/Runtime/Landscape/Private/LandscapePhysicalMaterial.cpp
--- a/Engine/Source/Runtime/Landscape/Private/LandscapePhysicalMaterial.cpp
+++ b/Engine/Source/Runtime/Landscape/Private/LandscapePhysicalMaterial.cpp
@@ -297,7 +297,7 @@ namespace
 
 		// Create the view
 		FSceneViewFamily::ConstructionValues ViewFamilyInit(nullptr, SceneInterface, FEngineShowFlags(ESFIM_Game));
-		ViewFamilyInit.SetWorldTimes(0.0f, 0.0f, 0.0f);
+		ViewFamilyInit.SetTime(FGameTime());
 		FSceneViewFamilyContext ViewFamily(ViewFamilyInit);
 		ViewFamily.LandscapeLODOverride = 0; // Force LOD 0 render
 
diff --git a/Engine/Source/Runtime/Renderer/Private/DecalRenderingShared.cpp b/Engine/Source/Runtime/Renderer/Private/DecalRenderingShared.cpp
--- a/Engine/Source/Runtime/Renderer/Private/DecalRenderingShared.cpp
+++ b/Engine/Source/Runtime/Renderer/Private/DecalRenderingShared.cpp
@@ -154,9 +154,9 @@ public:
 		float LifetimeAlpha = 1.0f;
 
 		// Certain engine captures (e.g. environment reflection) don't have a tick. Default to fully opaque.
-		if (View.Family->CurrentWorldTime)
+		if (View.Family->Time.GetWorldTimeSeconds())
 		{
-			LifetimeAlpha = FMath::Clamp(FMath::Min(View.Family->CurrentWorldTime * -DecalProxy.InvFadeDuration + DecalProxy.FadeStartDelayNormalized, View.Family->CurrentWorldTime * DecalProxy.InvFadeInDuration + DecalProxy.FadeInStartDelayNormalized), 0.0f, 1.0f);
+			LifetimeAlpha = FMath::Clamp(FMath::Min(View.Family->Time.GetWorldTimeSeconds() * -DecalProxy.InvFadeDuration + DecalProxy.FadeStartDelayNormalized, View.Family->Time.GetWorldTimeSeconds() * DecalProxy.InvFadeInDuration + DecalProxy.FadeInStartDelayNormalized), 0.0f, 1.0f);
 		}
  
 		SetShaderValue(RHICmdList, ShaderRHI, DecalParams, FVector2f(FadeAlphaValue, LifetimeAlpha));
diff --git a/Engine/Source/Runtime/Renderer/Private/DeferredShadingRenderer.cpp b/Engine/Source/Runtime/Renderer/Private/DeferredShadingRenderer.cpp
--- a/Engine/Source/Runtime/Renderer/Private/DeferredShadingRenderer.cpp
+++ b/Engine/Source/Runtime/Renderer/Private/DeferredShadingRenderer.cpp
@@ -549,7 +549,7 @@ bool FDeferredShadingSceneRenderer::GatherRayTracingWorldInstancesForView(FRDGBu
 		*View.RayTracingMeshResourceCollector
 	};
 
-	const float CurrentWorldTime = View.Family->CurrentWorldTime;
+	const float CurrentWorldTime = View.Family->Time.GetWorldTimeSeconds();
 
 	struct FRelevantPrimitive
 	{
diff --git a/Engine/Source/Runtime/Renderer/Private/GlobalDistanceField.cpp b/Engine/Source/Runtime/Renderer/Private/GlobalDistanceField.cpp
--- a/Engine/Source/Runtime/Renderer/Private/GlobalDistanceField.cpp
+++ b/Engine/Source/Runtime/Renderer/Private/GlobalDistanceField.cpp
@@ -547,7 +547,7 @@ void UpdateGlobalDistanceFieldViewOrigin(const FViewInfo& View, bool bLumenEnabl
 			FVector& CameraVelocityOffset = View.ViewState->GlobalDistanceFieldCameraVelocityOffset;
 			const FVector CameraVelocity = View.ViewMatrices.GetViewOrigin() - View.PrevViewInfo.ViewMatrices.GetViewOrigin();
 			// Framerate independent decay
-			CameraVelocityOffset = CameraVelocityOffset * FMath::Pow(GAOGlobalDistanceFieldCameraPositionVelocityOffsetDecay, View.Family->DeltaWorldTime) + CameraVelocity;
+			CameraVelocityOffset = CameraVelocityOffset * FMath::Pow(GAOGlobalDistanceFieldCameraPositionVelocityOffsetDecay, View.Family->Time.GetDeltaWorldTimeSeconds()) + CameraVelocity;
 
 			const FScene* Scene = (const FScene*)View.Family->Scene;
 			const int32 NumClipmaps = GetNumGlobalDistanceFieldClipmaps();
diff --git a/Engine/Source/Runtime/Renderer/Private/IndirectLightingCache.cpp b/Engine/Source/Runtime/Renderer/Private/IndirectLightingCache.cpp
--- a/Engine/Source/Runtime/Renderer/Private/IndirectLightingCache.cpp
+++ b/Engine/Source/Runtime/Renderer/Private/IndirectLightingCache.cpp
@@ -583,7 +583,7 @@ void FIndirectLightingCache::FinalizeUpdateInternal_RenderThread(FScene* Scene,
 
 		{
 			SCOPE_CYCLE_COUNTER(STAT_UpdateIndirectLightingCacheTransitions);
-			UpdateTransitionsOverTime(TransitionsOverTimeToUpdate, Renderer.ViewFamily.DeltaWorldTime);
+			UpdateTransitionsOverTime(TransitionsOverTimeToUpdate, Renderer.ViewFamily.Time.GetDeltaWorldTimeSeconds());
 		}
 
 		for (int32 PrimitiveIndex = 0; PrimitiveIndex < PrimitivesToUpdateStaticMeshes.Num(); ++PrimitiveIndex)
diff --git a/Engine/Source/Runtime/Renderer/Private/PostProcess/PostProcessEyeAdaptation.cpp b/Engine/Source/Runtime/Renderer/Private/PostProcess/PostProcessEyeAdaptation.cpp
--- a/Engine/Source/Runtime/Renderer/Private/PostProcess/PostProcessEyeAdaptation.cpp
+++ b/Engine/Source/Runtime/Renderer/Private/PostProcess/PostProcessEyeAdaptation.cpp
@@ -403,7 +403,7 @@ FEyeAdaptationParameters GetEyeAdaptationParameters(const FViewInfo& View, ERHIF
 	Parameters.MaxAverageLuminance = MaxAverageLuminance;
 	Parameters.ExposureCompensationSettings = ExposureCompensationSettings;
 	Parameters.ExposureCompensationCurve = ExposureCompensationCurve;
-	Parameters.DeltaWorldTime = View.Family->DeltaWorldTime;
+	Parameters.DeltaWorldTime = View.Family->Time.GetDeltaWorldTimeSeconds();
 	Parameters.ExposureSpeedUp = Settings.AutoExposureSpeedUp;
 	Parameters.ExposureSpeedDown = Settings.AutoExposureSpeedDown;
 	Parameters.HistogramScale = HistogramScale;
diff --git a/Engine/Source/Runtime/Renderer/Private/PostProcess/PostProcessTestImage.cpp b/Engine/Source/Runtime/Renderer/Private/PostProcess/PostProcessTestImage.cpp
--- a/Engine/Source/Runtime/Renderer/Private/PostProcess/PostProcessTestImage.cpp
+++ b/Engine/Source/Runtime/Renderer/Private/PostProcess/PostProcessTestImage.cpp
@@ -42,7 +42,7 @@ void AddTestImagePass(FRDGBuilder& GraphBuilder, const FViewInfo& View, FScreenP
 	PassParameters->Output = GetScreenPassTextureViewportParameters(Viewport);
 	PassParameters->ColorRemap = GetColorRemapParameters();
 	PassParameters->FrameNumber = ViewFamily.FrameNumber;
-	PassParameters->FrameTime = ViewFamily.CurrentRealTime;
+	PassParameters->FrameTime = ViewFamily.Time.GetRealTimeSeconds();
 	PassParameters->RenderTargets[0] = FRenderTargetBinding(OutputTexture, ERenderTargetLoadAction::EClear);
 
 	TShaderMapRef<FTestImagePS> PixelShader(View.ShaderMap);
diff --git a/Engine/Source/Runtime/Renderer/Private/ReflectionEnvironmentCapture.cpp b/Engine/Source/Runtime/Renderer/Private/ReflectionEnvironmentCapture.cpp
--- a/Engine/Source/Runtime/Renderer/Private/ReflectionEnvironmentCapture.cpp
+++ b/Engine/Source/Runtime/Renderer/Private/ReflectionEnvironmentCapture.cpp
@@ -1384,7 +1384,7 @@ void CaptureSceneIntoScratchCubemap(
 
 		if( bStaticSceneOnly )
 		{
-			ViewFamilyInit.SetWorldTimes( 0.0f, 0.0f, 0.0f );
+			ViewFamilyInit.SetTime(FGameTime());
 		}
 
 		FSceneViewFamilyContext ViewFamily( ViewFamilyInit );
diff --git a/Engine/Source/Runtime/Renderer/Private/SceneOcclusion.cpp b/Engine/Source/Runtime/Renderer/Private/SceneOcclusion.cpp
--- a/Engine/Source/Runtime/Renderer/Private/SceneOcclusion.cpp
+++ b/Engine/Source/Runtime/Renderer/Private/SceneOcclusion.cpp
@@ -1141,7 +1141,7 @@ static FViewOcclusionQueriesPerView AllocateOcclusionTests(const FScene* Scene,
 			FSceneViewState::ShadowKeyOcclusionQueryMap& ShadowOcclusionQueryMap = ViewState->ShadowOcclusionQueryMaps[QueryIndex];
 
 			// Clear primitives which haven't been visible recently out of the occlusion history, and reset old pending occlusion queries.
-			ViewState->TrimOcclusionHistory(ViewFamily.CurrentRealTime, ViewFamily.CurrentRealTime - GEngine->PrimitiveProbablyVisibleTime, ViewFamily.CurrentRealTime, ViewState->OcclusionFrameCounter);
+			ViewState->TrimOcclusionHistory(ViewFamily.Time.GetRealTimeSeconds(), ViewFamily.Time.GetRealTimeSeconds() - GEngine->PrimitiveProbablyVisibleTime, ViewFamily.Time.GetRealTimeSeconds(), ViewState->OcclusionFrameCounter);
 
 			// Give back all these occlusion queries to the pool.
 			ShadowOcclusionQueryMap.Reset();
diff --git a/Engine/Source/Runtime/Renderer/Private/SceneRendering.cpp b/Engine/Source/Runtime/Renderer/Private/SceneRendering.cpp
--- a/Engine/Source/Runtime/Renderer/Private/SceneRendering.cpp
+++ b/Engine/Source/Runtime/Renderer/Private/SceneRendering.cpp
@@ -2575,8 +2575,7 @@ FSceneRenderer::FSceneRenderer(const FSceneViewFamily* InViewFamily,FHitProxyCon
 	// rendering differences in materials.
 	if (bAnyViewIsLocked)
 	{
-		ViewFamily.CurrentRealTime = 0.0f;
-		ViewFamily.CurrentWorldTime = 0.0f;
+		ViewFamily.Time = FGameTime::CreateDilated(0.0f, ViewFamily.Time.GetDeltaRealTimeSeconds(), 0.0f, ViewFamily.Time.GetDeltaWorldTimeSeconds());
 	}
 	
 	if(HitProxyConsumer)
@@ -4640,7 +4639,7 @@ static void DisplayInternals(FRHICommandListImmediate& RHICmdList, FViewInfo& In
 		// could be 0
 		auto State = InView.ViewState;
 
-		FCanvas Canvas((FRenderTarget*)Family->RenderTarget, NULL, Family->CurrentRealTime, Family->CurrentWorldTime, Family->DeltaWorldTime, InView.GetFeatureLevel());
+		FCanvas Canvas((FRenderTarget*)Family->RenderTarget, NULL, Family->Time, InView.GetFeatureLevel());
 		Canvas.SetRenderTargetRect(FIntRect(0, 0, Family->RenderTarget->GetSizeXY().X, Family->RenderTarget->GetSizeXY().Y));
 
 
@@ -4714,7 +4713,7 @@ static void DisplayInternals(FRHICommandListImmediate& RHICmdList, FViewInfo& In
 		}
 
 		CANVAS_HEADER(TEXT("Family:"))
-		CANVAS_LINE(false, TEXT("  Time (Real/World/DeltaWorld): %.2f/%.2f/%.2f"), Family->CurrentRealTime, Family->CurrentWorldTime, Family->DeltaWorldTime)
+		CANVAS_LINE(false, TEXT("  Time (Real/World/DeltaWorld): %.2f/%.2f/%.2f"), Family->Time.GetRealTimeSeconds(), Family->Time.GetWorldTimeSeconds(), Family->Time.GetDeltaWorldTimeSeconds())
 		CANVAS_LINE(false, TEXT("  FrameNumber: %u"), Family->FrameNumber)
 		CANVAS_LINE(false, TEXT("  ExposureSettings: %s"), *Family->ExposureSettings.ToString())
 		CANVAS_LINE(false, TEXT("  GammaCorrection: %.2f"), Family->GammaCorrection)
diff --git a/Engine/Source/Runtime/Renderer/Private/ScreenPass.h b/Engine/Source/Runtime/Renderer/Private/ScreenPass.h
--- a/Engine/Source/Runtime/Renderer/Private/ScreenPass.h
+++ b/Engine/Source/Runtime/Renderer/Private/ScreenPass.h
@@ -648,7 +648,7 @@ FORCEINLINE void AddDrawCanvasPass(
 	check(Output.IsValid());
 
 	const FSceneViewFamily& ViewFamily = *View.Family;
-	FCanvas& Canvas = *FCanvas::Create(GraphBuilder, Output.Texture, nullptr, ViewFamily.CurrentRealTime, ViewFamily.CurrentWorldTime, ViewFamily.DeltaWorldTime, View.GetFeatureLevel());
+	FCanvas& Canvas = *FCanvas::Create(GraphBuilder, Output.Texture, nullptr, ViewFamily.Time, View.GetFeatureLevel());
 	Canvas.SetRenderTargetRect(Output.ViewRect);
 
 	Function(Canvas);
diff --git a/Engine/Source/Runtime/Renderer/Private/ShadowSetup.cpp b/Engine/Source/Runtime/Renderer/Private/ShadowSetup.cpp
--- a/Engine/Source/Runtime/Renderer/Private/ShadowSetup.cpp
+++ b/Engine/Source/Runtime/Renderer/Private/ShadowSetup.cpp
@@ -1662,7 +1662,7 @@ void FProjectedShadowInfo::AddSubjectPrimitive(FPrimitiveSceneInfo* PrimitiveSce
 		if (bShadowRelevance)
 		{
 			// Update the primitive component's last render time. Allows the component to update when using bCastWhenHidden.
-			const float CurrentWorldTime = Views[0].Family->CurrentWorldTime;
+			const float CurrentWorldTime = Views[0].Family->Time.GetWorldTimeSeconds();
 			PrimitiveSceneInfo->UpdateComponentLastRenderTime(CurrentWorldTime, /*bUpdateLastRenderTimeOnScreen=*/false);
 
 			if (PrimitiveSceneInfo->NeedsUniformBufferUpdate())
@@ -1891,7 +1891,7 @@ uint64 FProjectedShadowInfo::AddSubjectPrimitive_AnyThread(
 		}
 
 		// Update the primitive component's last render time. Allows the component to update when using bCastWhenHidden.
-		const float CurrentWorldTime = CurrentView->Family->CurrentWorldTime;
+		const float CurrentWorldTime = CurrentView->Family->Time.GetWorldTimeSeconds();
 		PrimitiveSceneInfo->UpdateComponentLastRenderTime(CurrentWorldTime, /*bUpdateLastRenderTimeOnScreen=*/false);
 
 		if (PrimitiveSceneInfo->NeedsUniformBufferUpdate())
@@ -3612,7 +3612,7 @@ void FSceneRenderer::CreateWholeSceneProjectedShadow(
 					ComputeWholeSceneShadowCacheModes(
 						LightSceneInfo,
 						ProjectedShadowInitializer.bOnePassPointLightShadow,
-						ViewFamily.CurrentRealTime,
+						ViewFamily.Time.GetRealTimeSeconds(),
 						MaxDesiredResolution,
 						FIntPoint(MaxShadowResolution, MaxShadowResolutionY),
 						Scene,
@@ -4718,7 +4718,7 @@ void FSceneRenderer::AddViewDependentWholeSceneShadowsForView(
 						{
 							ComputeViewDependentWholeSceneShadowCacheModes(
 								&LightSceneInfo,
-								ViewFamily.CurrentRealTime,
+								ViewFamily.Time.GetRealTimeSeconds(),
 								Scene,
 								ProjectedShadowInitializer,
 								ShadowBufferResolution,
@@ -5154,7 +5154,7 @@ void FSceneRenderer::AllocateShadowDepthTargets(FRHICommandListImmediate& RHICmd
 
 		for (auto& ShadowMapData : ShadowMapDatas)
 		{
-			if (ShadowMapData.ShadowMap.IsValid() && ViewFamily.CurrentRealTime - ShadowMapData.LastUsedTime > 2.0f)
+			if (ShadowMapData.ShadowMap.IsValid() && ViewFamily.Time.GetRealTimeSeconds() - ShadowMapData.LastUsedTime > 2.0f)
 			{
 				ShadowMapData.InvalidateCachedShadow();
 			}
diff --git a/Engine/Source/Runtime/Renderer/Private/TranslucentRendering.cpp b/Engine/Source/Runtime/Renderer/Private/TranslucentRendering.cpp
--- a/Engine/Source/Runtime/Renderer/Private/TranslucentRendering.cpp
+++ b/Engine/Source/Runtime/Renderer/Private/TranslucentRendering.cpp
@@ -323,7 +323,7 @@ FSeparateTranslucencyDimensions UpdateTranslucencyTimers(FRHICommandListImmediat
 						ViewState->SmoothedHalfResTranslucencyGPUDuration = FMath::Lerp(ViewState->SmoothedHalfResTranslucencyGPUDuration, LastFrameTranslucencyDurationMS, LerpAlpha);
 
 						// Don't re-asses switching for some time after the last switch
-						if (View.Family->CurrentRealTime - ViewState->LastAutoDownsampleChangeTime > CVarSeparateTranslucencyMinDownsampleChangeTime.GetValueOnRenderThread())
+						if (View.Family->Time.GetRealTimeSeconds() - ViewState->LastAutoDownsampleChangeTime > CVarSeparateTranslucencyMinDownsampleChangeTime.GetValueOnRenderThread())
 						{
 							// Downsample if the smoothed time is larger than the threshold
 							ViewState->bShouldAutoDownsampleTranslucency = ViewState->SmoothedHalfResTranslucencyGPUDuration > CVarSeparateTranslucencyDurationUpsampleThreshold.GetValueOnRenderThread();
@@ -341,7 +341,7 @@ FSeparateTranslucencyDimensions UpdateTranslucencyTimers(FRHICommandListImmediat
 						const float LerpAlpha = ViewState->SmoothedFullResTranslucencyGPUDuration == 0 ? 1.0f : .1f;
 						ViewState->SmoothedFullResTranslucencyGPUDuration = FMath::Lerp(ViewState->SmoothedFullResTranslucencyGPUDuration, LastFrameTranslucencyDurationMS, LerpAlpha);
 
-						if (View.Family->CurrentRealTime - ViewState->LastAutoDownsampleChangeTime > CVarSeparateTranslucencyMinDownsampleChangeTime.GetValueOnRenderThread())
+						if (View.Family->Time.GetRealTimeSeconds() - ViewState->LastAutoDownsampleChangeTime > CVarSeparateTranslucencyMinDownsampleChangeTime.GetValueOnRenderThread())
 						{
 							// Downsample if the smoothed time is larger than the threshold
 							ViewState->bShouldAutoDownsampleTranslucency = ViewState->SmoothedFullResTranslucencyGPUDuration > CVarSeparateTranslucencyDurationDownsampleThreshold.GetValueOnRenderThread();
@@ -355,7 +355,7 @@ FSeparateTranslucencyDimensions UpdateTranslucencyTimers(FRHICommandListImmediat
 
 					if (bOriginalShouldAutoDownsampleTranslucency != ViewState->bShouldAutoDownsampleTranslucency)
 					{
-						ViewState->LastAutoDownsampleChangeTime = View.Family->CurrentRealTime;
+						ViewState->LastAutoDownsampleChangeTime = View.Family->Time.GetRealTimeSeconds();
 					}
 
 					bAnyViewWantsDownsampledSeparateTranslucency = bAnyViewWantsDownsampledSeparateTranslucency || ViewState->bShouldAutoDownsampleTranslucency;
diff --git a/Engine/Source/Runtime/Renderer/Private/VT/RuntimeVirtualTextureRender.cpp b/Engine/Source/Runtime/Renderer/Private/VT/RuntimeVirtualTextureRender.cpp
--- a/Engine/Source/Runtime/Renderer/Private/VT/RuntimeVirtualTextureRender.cpp
+++ b/Engine/Source/Runtime/Renderer/Private/VT/RuntimeVirtualTextureRender.cpp
@@ -1115,7 +1115,7 @@ namespace RuntimeVirtualTexture
 		//todo[vt]: Some of this, such as ViewRotationMatrix, can be computed once in the Finalizer and passed down.
 		//todo[vt]: Have specific shader variations and setup for different output texture configs
 		FSceneViewFamily::ConstructionValues ViewFamilyInit(nullptr, Scene, FEngineShowFlags(ESFIM_Game));
-		ViewFamilyInit.SetWorldTimes(0.0f, 0.0f, 0.0f);
+		ViewFamilyInit.SetTime(FGameTime());
 		FSceneViewFamily& ViewFamily = *GraphBuilder.AllocObject<FSceneViewFamily>(ViewFamilyInit);
 
 		FSceneViewInitOptions ViewInitOptions;
diff --git a/Engine/Source/Runtime/SlateCore/Public/Rendering/SlateRenderer.h b/Engine/Source/Runtime/SlateCore/Public/Rendering/SlateRenderer.h
--- a/Engine/Source/Runtime/SlateCore/Public/Rendering/SlateRenderer.h
+++ b/Engine/Source/Runtime/SlateCore/Public/Rendering/SlateRenderer.h
@@ -33,6 +33,7 @@ struct FRenderThreadUpdateContext
 	float WorldTimeSeconds;
 	float DeltaTimeSeconds;
 	float RealTimeSeconds;
+	float DeltaRealTimeSeconds;
 	FRenderTarget* RenderTarget;
 	ISlate3DRenderer* Renderer;
 	bool bClearTarget;
diff --git a/Engine/Source/Runtime/SlateRHIRenderer/Private/Slate3DRenderer.cpp b/Engine/Source/Runtime/SlateRHIRenderer/Private/Slate3DRenderer.cpp
--- a/Engine/Source/Runtime/SlateRHIRenderer/Private/Slate3DRenderer.cpp
+++ b/Engine/Source/Runtime/SlateRHIRenderer/Private/Slate3DRenderer.cpp
@@ -158,7 +158,7 @@ void FSlate3DRenderer::DrawWindowToTarget_RenderThread(FRHICommandListImmediate&
 
 			FSlateBackBuffer BackBufferTarget(Context.RenderTarget->GetRenderTargetTexture(), FIntPoint(RTTextureRHI->GetSizeX(), RTTextureRHI->GetSizeY()));
 
-			FSlateRenderingParams DrawOptions(ProjectionMatrix, Context.WorldTimeSeconds, Context.DeltaTimeSeconds, Context.RealTimeSeconds);
+			FSlateRenderingParams DrawOptions(ProjectionMatrix, FGameTime::CreateDilated(Context.RealTimeSeconds, Context.DeltaRealTimeSeconds, Context.WorldTimeSeconds, Context.DeltaTimeSeconds));
 			// The scene renderer will handle it in this case
 			DrawOptions.bAllowSwitchVerticalAxis = false;
 			DrawOptions.ViewOffset = DrawOffset;
diff --git a/Engine/Source/Runtime/SlateRHIRenderer/Private/SlateRHIRenderer.cpp b/Engine/Source/Runtime/SlateRHIRenderer/Private/SlateRHIRenderer.cpp
--- a/Engine/Source/Runtime/SlateRHIRenderer/Private/SlateRHIRenderer.cpp
+++ b/Engine/Source/Runtime/SlateRHIRenderer/Private/SlateRHIRenderer.cpp
@@ -96,9 +96,7 @@ struct FSlateDrawWindowCommandParams
 #if WANTS_DRAW_MESH_EVENTS
 	FString WindowTitle;
 #endif
-	float WorldTimeSeconds;
-	float DeltaTimeSeconds;
-	float RealTimeSeconds;
+	FGameTime Time;
 	bool bLockToVsync;
 	bool bClear;
 };
@@ -888,7 +886,7 @@ void FSlateRHIRenderer::DrawWindow_RenderThread(FRHICommandListImmediate& RHICmd
 						{
 							FSlateBackBuffer BackBufferTarget(BackBuffer, FIntPoint(ViewportWidth, ViewportHeight));
 
-							FSlateRenderingParams RenderParams(ViewMatrix * ViewportInfo.ProjectionMatrix, DrawCommandParams.WorldTimeSeconds, DrawCommandParams.DeltaTimeSeconds, DrawCommandParams.RealTimeSeconds);
+							FSlateRenderingParams RenderParams(ViewMatrix * ViewportInfo.ProjectionMatrix, DrawCommandParams.Time);
 							RenderParams.bWireFrame = !!SlateWireFrame;
 							RenderParams.bIsHDR = ViewportInfo.bHDREnabled;
 
@@ -1324,9 +1322,9 @@ void FSlateRHIRenderer::DrawWindows_Private(FSlateDrawBuffer& WindowDrawBuffer)
 #else
 					Params.bClear = false;
 #endif	
-					Params.WorldTimeSeconds = FApp::GetCurrentTime() - GStartTime;
-					Params.DeltaTimeSeconds = FApp::GetDeltaTime();
-					Params.RealTimeSeconds = FPlatformTime::Seconds() - GStartTime;
+					Params.Time = FGameTime::CreateDilated(
+						FPlatformTime::Seconds() - GStartTime, FApp::GetDeltaTime(),
+						FApp::GetCurrentTime() - GStartTime, FApp::GetDeltaTime());
 
 					// Skip the actual draw if we're in a headless execution environment
 					bool bLocalTakingAScreenShot = bTakingAScreenShot;
diff --git a/Engine/Source/Runtime/SlateRHIRenderer/Private/SlateRHIRenderingPolicy.cpp b/Engine/Source/Runtime/SlateRHIRenderer/Private/SlateRHIRenderingPolicy.cpp
--- a/Engine/Source/Runtime/SlateRHIRenderer/Private/SlateRHIRenderingPolicy.cpp
+++ b/Engine/Source/Runtime/SlateRHIRenderer/Private/SlateRHIRenderingPolicy.cpp
@@ -630,7 +630,7 @@ void FSlateRHIRenderingPolicy::DrawElements(
 					ResourceManager->GetSceneAt(i),
 					DefaultShowFlags
 				)
-				.SetWorldTimes(Params.CurrentWorldTime, Params.DeltaTimeSeconds, Params.CurrentRealTime)
+				.SetTime(Params.Time)
 				.SetGammaCorrection(DisplayGamma)
 				.SetRealtimeUpdate(true)
 			);
@@ -645,7 +645,7 @@ void FSlateRHIRenderingPolicy::DrawElements(
 				nullptr,
 				DefaultShowFlags
 			)
-			.SetWorldTimes(Params.CurrentWorldTime, Params.DeltaTimeSeconds, Params.CurrentRealTime)
+			.SetTime(Params.Time)
 			.SetGammaCorrection(DisplayGamma)
 			.SetRealtimeUpdate(true)
 		);
diff --git a/Engine/Source/Runtime/SlateRHIRenderer/Private/SlateRHIRenderingPolicy.h b/Engine/Source/Runtime/SlateRHIRenderer/Private/SlateRHIRenderingPolicy.h
--- a/Engine/Source/Runtime/SlateRHIRenderer/Private/SlateRHIRenderingPolicy.h
+++ b/Engine/Source/Runtime/SlateRHIRenderer/Private/SlateRHIRenderingPolicy.h
@@ -28,19 +28,15 @@ struct FSlateRenderingParams
 {
 	FMatrix44f ViewProjectionMatrix;
 	FVector2D ViewOffset;
-	float CurrentWorldTime;
-	float DeltaTimeSeconds;
-	float CurrentRealTime;
+	FGameTime Time;
 	bool bAllowSwitchVerticalAxis;
 	bool bWireFrame;
 	bool bIsHDR;
 
-	FSlateRenderingParams(const FMatrix& InViewProjectionMatrix, float InCurrentWorldTime, float InDeltaTimeSeconds, float InCurrentRealTime)
+	FSlateRenderingParams(const FMatrix& InViewProjectionMatrix, FGameTime InTime)
 		: ViewProjectionMatrix(InViewProjectionMatrix)
 		, ViewOffset(0, 0)
-		, CurrentWorldTime(InCurrentWorldTime)
-		, DeltaTimeSeconds(InDeltaTimeSeconds)
-		, CurrentRealTime(InCurrentRealTime)
+		, Time(InTime)
 		, bAllowSwitchVerticalAxis(true)
 		, bWireFrame(false)
 		, bIsHDR(false)
diff --git a/Engine/Source/Runtime/UMG/Private/Components/Viewport.cpp b/Engine/Source/Runtime/UMG/Private/Components/Viewport.cpp
--- a/Engine/Source/Runtime/UMG/Private/Components/Viewport.cpp
+++ b/Engine/Source/Runtime/UMG/Private/Components/Viewport.cpp
@@ -130,26 +130,17 @@ void FUMGViewportClient::Draw(FViewport* InViewport, FCanvas* Canvas)
 	FViewport* ViewportBackup = Viewport;
 	Viewport = InViewport ? InViewport : Viewport;
 
-	// Determine whether we should use world time or real time based on the scene.
-	float TimeSeconds;
-	float RealTimeSeconds;
-	float DeltaTimeSeconds;
-
 	const bool bIsRealTime = true;
 
-	UWorld* World = GWorld;
+	UWorld* World = GetWorld();
+	FGameTime Time;
 	if ( bIsRealTime || GetScene() != World->Scene )
 	{
-		// Use time relative to start time to avoid issues with float vs double
-		TimeSeconds = FApp::GetCurrentTime() - GStartTime;
-		RealTimeSeconds = FApp::GetCurrentTime() - GStartTime;
-		DeltaTimeSeconds = FApp::GetDeltaTime();
+		Time = FGameTime::GetTimeSinceAppStart();
 	}
 	else
 	{
-		TimeSeconds = World->GetTimeSeconds();
-		RealTimeSeconds = World->GetRealTimeSeconds();
-		DeltaTimeSeconds = World->GetDeltaSeconds();
+		Time = World->GetTime();
 	}
 
 	// Setup a FSceneViewFamily/FSceneView for the viewport.
@@ -157,7 +148,7 @@ void FUMGViewportClient::Draw(FViewport* InViewport, FCanvas* Canvas)
 		Canvas->GetRenderTarget(),
 		GetScene(),
 		EngineShowFlags)
-		.SetWorldTimes(TimeSeconds, DeltaTimeSeconds, RealTimeSeconds)
+		.SetTime(Time)
 		.SetRealtimeUpdate(bIsRealTime));
 
 	// Get DPI derived view fraction.
diff --git a/Engine/Source/Runtime/UMG/Private/Slate/WidgetRenderer.cpp b/Engine/Source/Runtime/UMG/Private/Slate/WidgetRenderer.cpp
--- a/Engine/Source/Runtime/UMG/Private/Slate/WidgetRenderer.cpp
+++ b/Engine/Source/Runtime/UMG/Private/Slate/WidgetRenderer.cpp
@@ -263,6 +263,7 @@ void FWidgetRenderer::DrawWindow(
 			static_cast<float>(FApp::GetCurrentTime() - GStartTime),
 			static_cast<float>(FApp::GetDeltaTime()),
 			static_cast<float>(FPlatformTime::Seconds() - GStartTime),
+			static_cast<float>(FApp::GetDeltaTime()),
 			RenderTarget,
 			Renderer.Get(),
 			bClearTarget
@@ -322,6 +323,7 @@ bool FWidgetRenderer::DrawInvalidationRoot(TSharedRef<SVirtualWindow>& VirtualWi
 				static_cast<float>(FApp::GetCurrentTime() - GStartTime),
 				static_cast<float>(FApp::GetDeltaTime()),
 				static_cast<float>(FPlatformTime::Seconds() - GStartTime),
+				static_cast<float>(FApp::GetDeltaTime()),
 				static_cast<FRenderTarget*>(RenderTarget->GameThread_GetRenderTargetResource()),
 				Renderer.Get(),
 				bClearTarget
