Added proper ViewExtensions setup to scene captures performed by FMaterialUtilities::ExportLandscapeMaterial()
* Fixes USD export automated test

diff --git a/Engine/Source/Developer/MaterialUtilities/Private/MaterialUtilities.cpp b/Engine/Source/Developer/MaterialUtilities/Private/MaterialUtilities.cpp
--- a/Engine/Source/Developer/MaterialUtilities/Private/MaterialUtilities.cpp
+++ b/Engine/Source/Developer/MaterialUtilities/Private/MaterialUtilities.cpp
@@ -19,6 +19,7 @@
 #include "Engine/TextureCube.h"
 #include "Engine/Texture2DArray.h"
 #include "SceneView.h"
+#include "SceneViewExtension.h"
 #include "RendererInterface.h"
 #include "EngineModule.h"
 #include "ImageUtils.h"
@@ -776,6 +777,14 @@ static void RenderSceneToTexture(
 
 	FCanvas Canvas(RenderTargetResource, NULL, FGameTime::GetTimeSinceAppStart(), Scene->GetFeatureLevel());
 	Canvas.Clear(FLinearColor::Transparent);
+
+	ViewFamily.ViewExtensions = GEngine->ViewExtensions->GatherActiveExtensions(FSceneViewExtensionContext(Scene));
+	for (const FSceneViewExtensionRef& Extension : ViewFamily.ViewExtensions)
+	{
+		Extension->SetupViewFamily(ViewFamily);
+		Extension->SetupView(ViewFamily, *NewView);
+	}
+
 	GetRendererModule().BeginRenderingViewFamily(&Canvas, &ViewFamily);
 
 	// Copy the contents of the remote texture to system memory
