Remove as much C++ deprecation as possible up to 4.17 (along with a few scattered removals from beyond)

diff --git a/Engine/Config/BaseEngine.ini b/Engine/Config/BaseEngine.ini
--- a/Engine/Config/BaseEngine.ini
+++ b/Engine/Config/BaseEngine.ini
@@ -1476,7 +1476,6 @@ FOVAngle=90.000000
 GodMode=True
 UseAxisIndicator=True
 HeightMapExportClassName="TerrainHeightMapExporterTextT3D"
-bGroupingActive=true
 bCustomCameraAlignEmitter=true
 CustomCameraAlignEmitterDistance=100.0
 bDrawSocketsInGMode=false
diff --git a/Engine/Plugins/Runtime/GameplayAbilities/Source/GameplayAbilities/Private/AbilitySystemComponent.cpp b/Engine/Plugins/Runtime/GameplayAbilities/Source/GameplayAbilities/Private/AbilitySystemComponent.cpp
--- a/Engine/Plugins/Runtime/GameplayAbilities/Source/GameplayAbilities/Private/AbilitySystemComponent.cpp
+++ b/Engine/Plugins/Runtime/GameplayAbilities/Source/GameplayAbilities/Private/AbilitySystemComponent.cpp
@@ -493,12 +493,6 @@ FActiveGameplayEffectHandle UAbilitySystemComponent::ApplyGameplayEffectToSelf(c
 	return FActiveGameplayEffectHandle();
 }
 
-FOnActiveGameplayEffectRemoved* UAbilitySystemComponent::OnGameplayEffectRemovedDelegate(FActiveGameplayEffectHandle Handle)
-{
-	FActiveGameplayEffect* ActiveEffect = ActiveGameplayEffects.GetActiveGameplayEffect(Handle);
-	return ActiveEffect ? &ActiveEffect->EventSet.DEPRECATED_OnEffectRemoved : nullptr;
-}
-
 FActiveGameplayEffectEvents* UAbilitySystemComponent::GetActiveEffectEventSet(FActiveGameplayEffectHandle Handle)
 {
 	FActiveGameplayEffect* ActiveEffect = ActiveGameplayEffects.GetActiveGameplayEffect(Handle);
diff --git a/Engine/Plugins/Runtime/GameplayAbilities/Source/GameplayAbilities/Private/GameplayEffect.cpp b/Engine/Plugins/Runtime/GameplayAbilities/Source/GameplayAbilities/Private/GameplayEffect.cpp
--- a/Engine/Plugins/Runtime/GameplayAbilities/Source/GameplayAbilities/Private/GameplayEffect.cpp
+++ b/Engine/Plugins/Runtime/GameplayAbilities/Source/GameplayAbilities/Private/GameplayEffect.cpp
@@ -3514,7 +3514,6 @@ void FActiveGameplayEffectsContainer::InternalOnActiveGameplayEffectRemoved(FAct
 		ABILITY_LOG(Warning, TEXT("InternalOnActiveGameplayEffectRemoved called with no GameplayEffect: %s"), *Effect.Handle.ToString());
 	}
 
-	Effect.EventSet.DEPRECATED_OnEffectRemoved.Broadcast();
 	Effect.EventSet.OnEffectRemoved.Broadcast(GameplayEffectRemovalInfo);
 
 	OnActiveGameplayEffectRemovedDelegate.Broadcast(Effect);
diff --git a/Engine/Plugins/Runtime/GameplayAbilities/Source/GameplayAbilities/Public/Abilities/GameplayAbilityWorldReticle.h b/Engine/Plugins/Runtime/GameplayAbilities/Source/GameplayAbilities/Public/Abilities/GameplayAbilityWorldReticle.h
--- a/Engine/Plugins/Runtime/GameplayAbilities/Source/GameplayAbilities/Public/Abilities/GameplayAbilityWorldReticle.h
+++ b/Engine/Plugins/Runtime/GameplayAbilities/Source/GameplayAbilities/Public/Abilities/GameplayAbilityWorldReticle.h
@@ -33,10 +33,6 @@ public:
 
 	virtual void Tick(float DeltaSeconds) override;
 
-	/** Accessor for checking, before instantiating, if this WorldReticle will replicate. */
-	UE_DEPRECATED(4.12, "Call AActor::GetIsReplicated instead")
-	bool GetReplicates() const { return GetIsReplicated(); }
-
 	// ------------------------------
 
 	virtual bool IsNetRelevantFor(const AActor* RealViewer, const AActor* ViewTarget, const FVector& SrcLocation) const override;
diff --git a/Engine/Plugins/Runtime/GameplayAbilities/Source/GameplayAbilities/Public/AbilitySystemComponent.h b/Engine/Plugins/Runtime/GameplayAbilities/Source/GameplayAbilities/Public/AbilitySystemComponent.h
--- a/Engine/Plugins/Runtime/GameplayAbilities/Source/GameplayAbilities/Public/AbilitySystemComponent.h
+++ b/Engine/Plugins/Runtime/GameplayAbilities/Source/GameplayAbilities/Public/AbilitySystemComponent.h
@@ -502,9 +502,6 @@ class GAMEPLAYABILITIES_API UAbilitySystemComponent : public UGameplayTasksCompo
 	/** Called when any gameplay effects are removed */
 	FOnGivenActiveGameplayEffectRemoved& OnAnyGameplayEffectRemovedDelegate();
 
-	UE_DEPRECATED(4.17, "Use OnGameplayEffectRemoved_InfoDelegate (the delegate signature has changed)")
-	FOnActiveGameplayEffectRemoved* OnGameplayEffectRemovedDelegate(FActiveGameplayEffectHandle Handle);
-
 	/** Returns delegate structure that allows binding to several gameplay effect changes */
 	FActiveGameplayEffectEvents* GetActiveEffectEventSet(FActiveGameplayEffectHandle Handle);
 	FOnActiveGameplayEffectRemoved_Info* OnGameplayEffectRemoved_InfoDelegate(FActiveGameplayEffectHandle Handle);
diff --git a/Engine/Plugins/Runtime/GameplayAbilities/Source/GameplayAbilities/Public/GameplayEffectAggregator.h b/Engine/Plugins/Runtime/GameplayAbilities/Source/GameplayAbilities/Public/GameplayEffectAggregator.h
--- a/Engine/Plugins/Runtime/GameplayAbilities/Source/GameplayAbilities/Public/GameplayEffectAggregator.h
+++ b/Engine/Plugins/Runtime/GameplayAbilities/Source/GameplayAbilities/Public/GameplayEffectAggregator.h
@@ -148,9 +148,6 @@ struct GAMEPLAYABILITIES_API FAggregatorModChannel
 	 */
 	void GetAllAggregatorMods(EGameplayModEvaluationChannel Channel, OUT TMap<EGameplayModEvaluationChannel, const TArray<FAggregatorMod>*>& OutMods) const;
 
-	UE_DEPRECATED(4.17, "Use GetAllAggregatorMods")
-	void DebugGetAllAggregatorMods(EGameplayModEvaluationChannel Channel, OUT TMap<EGameplayModEvaluationChannel, const TArray<FAggregatorMod>*>& OutMods) const { return GetAllAggregatorMods(Channel, OutMods); }
-	
 	/**
 	 * Called when the mod channel's gameplay effect dependencies have potentially been swapped out for new ones, like when GE arrays are cloned.
 	 * Updates mod handles appropriately.
@@ -260,9 +257,6 @@ struct GAMEPLAYABILITIES_API FAggregatorModChannelContainer
 	 */
 	void GetAllAggregatorMods(OUT TMap<EGameplayModEvaluationChannel, const TArray<FAggregatorMod>*>& OutMods) const;
 
-	UE_DEPRECATED(4.17, "Use GetAllAggregatorMods")
-	void DebugGetAllAggregatorMods(OUT TMap<EGameplayModEvaluationChannel, const TArray<FAggregatorMod>*>& OutMods) const { return GetAllAggregatorMods(OutMods); }
-
 	/**
 	 * Called when the container's gameplay effect dependencies have potentially been swapped out for new ones, like when GE arrays are cloned.
 	 * Updates each channel appropriately.
@@ -343,9 +337,6 @@ struct GAMEPLAYABILITIES_API FAggregator : public TSharedFromThis<FAggregator>
 	 */
 	void GetAllAggregatorMods(OUT TMap<EGameplayModEvaluationChannel, const TArray<FAggregatorMod>*>& OutMods) const;
 
-	UE_DEPRECATED(4.17, "Use GetAllAggregatorMods")
-	void DebugGetAllAggregatorMods(OUT TMap<EGameplayModEvaluationChannel, const TArray<FAggregatorMod>*>& OutMods) const { return GetAllAggregatorMods(OutMods); }
-	
 	/**
 	 * Called when the aggregator's gameplay effect dependencies have potentially been swapped out for new ones, like when GE arrays are cloned.
 	 * Updates the modifier channel container appropriately, as well as directly-specified dependents.
diff --git a/Engine/Plugins/Runtime/GameplayAbilities/Source/GameplayAbilities/Public/GameplayEffectTypes.h b/Engine/Plugins/Runtime/GameplayAbilities/Source/GameplayAbilities/Public/GameplayEffectTypes.h
--- a/Engine/Plugins/Runtime/GameplayAbilities/Source/GameplayAbilities/Public/GameplayEffectTypes.h
+++ b/Engine/Plugins/Runtime/GameplayAbilities/Source/GameplayAbilities/Public/GameplayEffectTypes.h
@@ -1058,7 +1058,6 @@ DECLARE_DELEGATE(FDeferredTagChangeDelegate);
 DECLARE_MULTICAST_DELEGATE_TwoParams(FOnGameplayEffectTagCountChanged, const FGameplayTag, int32);
 DECLARE_MULTICAST_DELEGATE_OneParam(FOnGivenActiveGameplayEffectRemoved, const FActiveGameplayEffect&);
 
-DECLARE_MULTICAST_DELEGATE(FOnActiveGameplayEffectRemoved);
 DECLARE_MULTICAST_DELEGATE_OneParam(FOnActiveGameplayEffectRemoved_Info, const FGameplayEffectRemovalInfo&);
 DECLARE_MULTICAST_DELEGATE_ThreeParams(FOnActiveGameplayEffectStackChange, FActiveGameplayEffectHandle, int32 /*NewStackCount*/, int32 /*PreviousStackCount*/);
 DECLARE_MULTICAST_DELEGATE_ThreeParams(FOnActiveGameplayEffectTimeChange, FActiveGameplayEffectHandle, float /*NewStartTime*/, float /*NewDuration*/);
@@ -1067,7 +1066,6 @@ DECLARE_MULTICAST_DELEGATE_TwoParams(FOnActiveGameplayEffectInhibitionChanged, F
 /** Callback struct for different types of gameplay effect changes */
 struct FActiveGameplayEffectEvents
 {
-	FOnActiveGameplayEffectRemoved DEPRECATED_OnEffectRemoved;
 	FOnActiveGameplayEffectRemoved_Info OnEffectRemoved;
 	FOnActiveGameplayEffectStackChange OnStackChanged;
 	FOnActiveGameplayEffectTimeChange OnTimeChanged;
diff --git a/Engine/Source/Developer/AssetTools/Private/AssetTools.cpp b/Engine/Source/Developer/AssetTools/Private/AssetTools.cpp
--- a/Engine/Source/Developer/AssetTools/Private/AssetTools.cpp
+++ b/Engine/Source/Developer/AssetTools/Private/AssetTools.cpp
@@ -527,11 +527,6 @@ UObject* UAssetToolsImpl::CreateAsset(const FString& AssetName, const FString& P
 	return NewObj;
 }
 
-UObject* UAssetToolsImpl::CreateAsset(UClass* AssetClass, UFactory* Factory, FName CallingContext /*= NAME_None*/)
-{
-	return CreateAssetWithDialog(AssetClass, Factory, CallingContext);
-}
-
 UObject* UAssetToolsImpl::CreateAssetWithDialog(UClass* AssetClass, UFactory* Factory, FName CallingContext)
 {
 	if (Factory != nullptr)
@@ -1103,11 +1098,6 @@ void UAssetToolsImpl::RenameReferencingSoftObjectPaths(const TArray<UPackage *>
 	AssetRenameManager->RenameReferencingSoftObjectPaths(PackagesToCheck, AssetRedirectorMap);
 }
 
-TArray<UObject*> UAssetToolsImpl::ImportAssets(const FString& DestinationPath)
-{
-	return ImportAssetsWithDialog(DestinationPath);
-}
-
 TArray<UObject*> UAssetToolsImpl::ImportAssetsWithDialog(const FString& DestinationPath)
 {
 	const bool bAllowAsyncImport = false;
diff --git a/Engine/Source/Developer/AssetTools/Private/AssetTools.h b/Engine/Source/Developer/AssetTools/Private/AssetTools.h
--- a/Engine/Source/Developer/AssetTools/Private/AssetTools.h
+++ b/Engine/Source/Developer/AssetTools/Private/AssetTools.h
@@ -75,7 +75,6 @@ public:
 	virtual void GetClassTypeActionsList( TArray<TWeakPtr<IClassTypeActions>>& OutClassTypeActionsList ) const override;
 	virtual TWeakPtr<IClassTypeActions> GetClassTypeActionsForClass( UClass* Class ) const override;
 	virtual UObject* CreateAsset(const FString& AssetName, const FString& PackagePath, UClass* AssetClass, UFactory* Factory, FName CallingContext = NAME_None) override;
-	virtual UObject* CreateAsset(UClass* AssetClass, UFactory* Factory, FName CallingContext = NAME_None) override;
 	virtual UObject* CreateAssetWithDialog(UClass* AssetClass, UFactory* Factory, FName CallingContext = NAME_None) override;
 	virtual UObject* CreateAssetWithDialog(const FString& AssetName, const FString& PackagePath, UClass* AssetClass, UFactory* Factory, FName CallingContext = NAME_None, const bool bCallConfigureProperties = true) override;
 	virtual UObject* DuplicateAsset(const FString& AssetName, const FString& PackagePath, UObject* OriginalObject) override;
@@ -86,7 +85,6 @@ public:
 	virtual void FindSoftReferencesToObject(FSoftObjectPath TargetObject, TArray<UObject*>& ReferencingObjects) override;
 	virtual void FindSoftReferencesToObjects(const TArray<FSoftObjectPath>& TargetObjects, TMap<FSoftObjectPath, TArray<UObject*>>& ReferencingObjects) override;
 	virtual void RenameReferencingSoftObjectPaths(const TArray<UPackage *> PackagesToCheck, const TMap<FSoftObjectPath, FSoftObjectPath>& AssetRedirectorMap) override;
-	virtual TArray<UObject*> ImportAssets(const FString& DestinationPath) override;
 	virtual TArray<UObject*> ImportAssetsWithDialog(const FString& DestinationPath) override;
 	virtual void ImportAssetsWithDialogAsync(const FString& DestinationPath) override;
 	virtual TArray<UObject*> ImportAssets(const TArray<FString>& Files, const FString& DestinationPath, UFactory* ChosenFactory, bool bSyncToBrowser = true, TArray<TPair<FString, FString>>* FilesAndDestinations = nullptr, bool bAllowAsyncImport = false) const override;
diff --git a/Engine/Source/Developer/AssetTools/Public/IAssetTools.h b/Engine/Source/Developer/AssetTools/Public/IAssetTools.h
--- a/Engine/Source/Developer/AssetTools/Public/IAssetTools.h
+++ b/Engine/Source/Developer/AssetTools/Public/IAssetTools.h
@@ -225,10 +225,6 @@ public:
 	UFUNCTION(BlueprintCallable, Category = "Editor Scripting | Asset Tools")
 	virtual UObject* CreateAssetWithDialog(const FString& AssetName, const FString& PackagePath, UClass* AssetClass, UFactory* Factory, FName CallingContext = NAME_None, const bool bCallConfigureProperties = true) = 0;
 
-	/** Opens an asset picker dialog and creates an asset with the chosen path */
-	UE_DEPRECATED(4.17, "This version of CreateAsset has been deprecated.  Use CreateAssetWithDialog instead")
-	virtual UObject* CreateAsset(UClass* AssetClass, UFactory* Factory, FName CallingContext = NAME_None) = 0;
-
 	/** Opens an asset picker dialog and creates an asset with the path chosen in the dialog */
 	virtual UObject* CreateAssetWithDialog(UClass* AssetClass, UFactory* Factory, FName CallingContext = NAME_None) = 0;
 
@@ -274,9 +270,6 @@ public:
 	/** Event issued at the end of the rename process */
 	virtual FAssetPostRenameEvent& OnAssetPostRename() = 0;
 
-	UE_DEPRECATED(4.17, "This version of ImportAssets has been deprecated.  Use ImportAssetsWithDialog instead")
-	virtual TArray<UObject*> ImportAssets(const FString& DestinationPath) = 0;
-
 	/**
 	 * Opens a file open dialog to choose files to import to the destination path.
 	 *
diff --git a/Engine/Source/Developer/FunctionalTesting/Classes/FunctionalTestingManager.h b/Engine/Source/Developer/FunctionalTesting/Classes/FunctionalTestingManager.h
--- a/Engine/Source/Developer/FunctionalTesting/Classes/FunctionalTestingManager.h
+++ b/Engine/Source/Developer/FunctionalTesting/Classes/FunctionalTestingManager.h
@@ -45,9 +45,6 @@ class UFunctionalTestingManager : public UBlueprintFunctionLibrary
 	UFUNCTION(BlueprintCallable, Category="FunctionalTesting", meta=(WorldContext="WorldContextObject", CallableWithoutWorldContext ) )
 	static bool RunAllFunctionalTests(UObject* WorldContextObject, bool bNewLog = true, bool bRunLooped = false, FString FailedTestsReproString = TEXT(""));
 
-	UE_DEPRECATED(4.18, "This function is deprecated and is no longer used. Please use the version without the bWaitForNavigationBuildFinish parameter.")
-	static bool RunAllFunctionalTests(UObject* WorldContextObject, bool bNewLog, bool bRunLooped, bool bWaitForNavigationBuildFinish, FString FailedTestsReproString = TEXT("")) { return RunAllFunctionalTests(WorldContextObject, bNewLog, bRunLooped, FailedTestsReproString); }
-
 	bool IsRunning() const { return bIsRunning; }
 	bool IsFinished() const { return bFinished; }
 	bool IsLooped() const { return bLooped; }
diff --git a/Engine/Source/Developer/MaterialUtilities/Private/MaterialUtilities.cpp b/Engine/Source/Developer/MaterialUtilities/Private/MaterialUtilities.cpp
--- a/Engine/Source/Developer/MaterialUtilities/Private/MaterialUtilities.cpp
+++ b/Engine/Source/Developer/MaterialUtilities/Private/MaterialUtilities.cpp
@@ -854,75 +854,6 @@ bool FMaterialUtilities::SupportsExport(EBlendMode InBlendMode, EMaterialPropert
 	return FExportMaterialProxy::WillFillData(InBlendMode, InMaterialProperty);
 }
 
-bool FMaterialUtilities::ExportMaterialProperty(UWorld* InWorld, UMaterialInterface* InMaterial, EMaterialProperty InMaterialProperty, UTextureRenderTarget2D* InRenderTarget, TArray<FColor>& OutBMP)
-{
-	FBox2D DummyBounds(FVector2D(0, 0), FVector2D(1, 1));
-	TArray<FVector2D> EmptyTexCoords;
-	FMaterialMergeData MaterialData(InMaterial, nullptr, nullptr, 0, DummyBounds, EmptyTexCoords);
-	const bool bForceGamma = (InMaterialProperty == MP_Normal) || (InMaterialProperty == MP_OpacityMask) || (InMaterialProperty == MP_Opacity);	
-
-	FIntPoint MaxSize = FindMaxTextureSize(InMaterial);
-	FIntPoint OutSize = MaxSize;
-	PRAGMA_DISABLE_DEPRECATION_WARNINGS
-	return RenderMaterialPropertyToTexture(MaterialData, InMaterialProperty, bForceGamma, PF_B8G8R8A8, MaxSize, OutSize, OutBMP);
-	PRAGMA_ENABLE_DEPRECATION_WARNINGS
-}
-
-bool FMaterialUtilities::ExportMaterialProperty(UWorld* InWorld, UMaterialInterface* InMaterial, EMaterialProperty InMaterialProperty, FIntPoint& OutSize, TArray<FColor>& OutBMP)
-{
-	FBox2D DummyBounds(FVector2D(0, 0), FVector2D(1, 1));
-	TArray<FVector2D> EmptyTexCoords;
-	FMaterialMergeData MaterialData(InMaterial, nullptr, nullptr, 0, DummyBounds, EmptyTexCoords);
-	const bool bForceGamma = (InMaterialProperty == MP_Normal) || (InMaterialProperty == MP_OpacityMask) || (InMaterialProperty == MP_Opacity);
-	OutSize = FindMaxTextureSize(InMaterial);
-	PRAGMA_DISABLE_DEPRECATION_WARNINGS
-	return RenderMaterialPropertyToTexture(MaterialData, InMaterialProperty, bForceGamma, PF_B8G8R8A8, OutSize, OutSize, OutBMP);
-	PRAGMA_ENABLE_DEPRECATION_WARNINGS
-}
-
-bool FMaterialUtilities::ExportMaterialProperty(UMaterialInterface* InMaterial, EMaterialProperty InMaterialProperty, TArray<FColor>& OutBMP, FIntPoint& OutSize)
-{
-	FBox2D DummyBounds(FVector2D(0, 0), FVector2D(1, 1));
-	TArray<FVector2D> EmptyTexCoords;
-	FMaterialMergeData MaterialData(InMaterial, nullptr, nullptr, 0, DummyBounds, EmptyTexCoords);
-	const bool bForceGamma = (InMaterialProperty == MP_Normal) || (InMaterialProperty == MP_OpacityMask) || (InMaterialProperty == MP_Opacity);
-	OutSize = FindMaxTextureSize(InMaterial);
-	PRAGMA_DISABLE_DEPRECATION_WARNINGS
-	return RenderMaterialPropertyToTexture(MaterialData, InMaterialProperty, bForceGamma, PF_B8G8R8A8, OutSize, OutSize, OutBMP);
-	PRAGMA_ENABLE_DEPRECATION_WARNINGS
-}
-
-bool FMaterialUtilities::ExportMaterialProperty(UMaterialInterface* InMaterial, EMaterialProperty InMaterialProperty, FIntPoint InSize, TArray<FColor>& OutBMP)
-{
-	FBox2D DummyBounds(FVector2D(0, 0), FVector2D(1, 1));
-	TArray<FVector2D> EmptyTexCoords;
-	FMaterialMergeData MaterialData(InMaterial, nullptr, nullptr, 0, DummyBounds, EmptyTexCoords);
-	const bool bForceGamma = (InMaterialProperty == MP_Normal) || (InMaterialProperty == MP_OpacityMask) || (InMaterialProperty == MP_Opacity);
-	FIntPoint OutSize;
-	PRAGMA_DISABLE_DEPRECATION_WARNINGS
-	return RenderMaterialPropertyToTexture(MaterialData, InMaterialProperty, bForceGamma, PF_B8G8R8A8, InSize, OutSize, OutBMP);
-	PRAGMA_ENABLE_DEPRECATION_WARNINGS
-}
-
-bool FMaterialUtilities::ExportMaterial(UWorld* InWorld, UMaterialInterface* InMaterial, FFlattenMaterial& OutFlattenMaterial)
-{
-	PRAGMA_DISABLE_DEPRECATION_WARNINGS
-	return ExportMaterial(InMaterial, OutFlattenMaterial);
-	PRAGMA_ENABLE_DEPRECATION_WARNINGS
-}
-
-bool FMaterialUtilities::ExportMaterial(UMaterialInterface* InMaterial, FFlattenMaterial& OutFlattenMaterial, struct FExportMaterialProxyCache* ProxyCache)
-{
-	FBox2D DummyBounds(FVector2D(0, 0), FVector2D(1, 1));
-	TArray<FVector2D> EmptyTexCoords;
-
-	FMaterialMergeData MaterialData(InMaterial, nullptr, nullptr, 0, DummyBounds, EmptyTexCoords);
-	PRAGMA_DISABLE_DEPRECATION_WARNINGS
-	ExportMaterial(MaterialData, OutFlattenMaterial, ProxyCache);
-	PRAGMA_ENABLE_DEPRECATION_WARNINGS
-	return true;
-}
-
 bool FMaterialUtilities::ExportLandscapeMaterial(ALandscapeProxy* InLandscape, const TSet<FPrimitiveComponentId>& HiddenPrimitives, FFlattenMaterial& OutFlattenMaterial)
 {
 	check(InLandscape);
@@ -2267,231 +2198,6 @@ bool FMaterialUtilities::ExportMaterialUVDensities(UMaterialInterface* InMateria
 	return true;
 }
 
-bool FMaterialUtilities::ExportMaterials(TArray<FMaterialMergeData*>& MergeData, TArray<FFlattenMaterial*>& OutFlattenMaterials)
-{
-	checkf(MergeData.Num() == OutFlattenMaterials.Num(), TEXT("Mismatch in number of input structures"));
-
-	TArray<FMeshData*> MeshSettings;
-	TArray<FMaterialData*> MatSettings;
-
-	const int32 NumMaterials = MergeData.Num();
-	for ( int32 MaterialIndex = 0; MaterialIndex < NumMaterials; ++MaterialIndex )
-	{
-		FMaterialMergeData* CurrentMergeData = MergeData[MaterialIndex];
-
-		FMeshData* MeshSet = new FMeshData();
-		MeshSet->MeshDescription = CurrentMergeData->Mesh;
-		MeshSet->TextureCoordinateBox = CurrentMergeData->TexcoordBounds;
-		MeshSet->CustomTextureCoordinates = CurrentMergeData->TexCoords;
-		MeshSettings.Add(MeshSet);
-
-		UMaterialInterface* Material = CurrentMergeData->Material;
-		UE_LOG(LogMaterialUtilities, Log, TEXT("Flattening material: %s"), *Material->GetName());
-
-		TArray<UTexture*> MaterialTextures;
-		Material->GetUsedTextures(MaterialTextures, EMaterialQualityLevel::Num, true, GMaxRHIFeatureLevel, true);		
-
-		FTextureCompilingManager::Get().FinishCompilation(MaterialTextures);
-
-		// Precache all used textures, otherwise could get everything rendered with low-res textures.
-		for (UTexture* Texture : MaterialTextures)
-		{
-			if (Texture != NULL)
-			{
-				UTexture2D* Texture2D = Cast<UTexture2D>(Texture);
-				if (Texture2D)
-				{
-					Texture2D->SetForceMipLevelsToBeResident(30.0f, true);
-					Texture2D->WaitForStreaming();
-				}
-			}
-		}
-
-		FFlattenMaterial* FlattenMaterial = OutFlattenMaterials[MaterialIndex];
-		FlattenMaterial->MaterialId = Material->GetLightingGuid();
-
-		// Determine whether or not certain properties can be rendered
-		const bool bRenderNormal = (Material->GetMaterial()->HasNormalConnected() || Material->GetMaterial()->bUseMaterialAttributes) && FlattenMaterial->ShouldGenerateDataForProperty(EFlattenMaterialProperties::Normal);
-		const bool bRenderTangent = (Material->GetMaterial()->Tangent.IsConnected() || Material->GetMaterial()->bUseMaterialAttributes) && FlattenMaterial->ShouldGenerateDataForProperty(EFlattenMaterialProperties::Tangent);
-		const bool bRenderEmissive = (Material->GetMaterial()->EmissiveColor.IsConnected() || Material->GetMaterial()->bUseMaterialAttributes) && FlattenMaterial->ShouldGenerateDataForProperty(EFlattenMaterialProperties::Emissive);
-		const bool bRenderOpacityMask = Material->IsPropertyActive(MP_OpacityMask) && Material->GetBlendMode() == BLEND_Masked && FlattenMaterial->ShouldGenerateDataForProperty(EFlattenMaterialProperties::OpacityMask);
-		const bool bRenderOpacity = Material->IsPropertyActive(MP_Opacity) && IsTranslucentBlendMode(Material->GetBlendMode()) && FlattenMaterial->ShouldGenerateDataForProperty(EFlattenMaterialProperties::Opacity);
-		const bool bRenderSubSurface = Material->IsPropertyActive(MP_SubsurfaceColor) && FlattenMaterial->ShouldGenerateDataForProperty(EFlattenMaterialProperties::SubSurface);
-		const bool bRenderMetallic = Material->IsPropertyActive(MP_Metallic) && FlattenMaterial->ShouldGenerateDataForProperty(EFlattenMaterialProperties::Metallic);
-		const bool bRenderSpecular = Material->IsPropertyActive(MP_Specular) && FlattenMaterial->ShouldGenerateDataForProperty(EFlattenMaterialProperties::Specular);
-		const bool bRenderRoughness = Material->IsPropertyActive(MP_Roughness) && FlattenMaterial->ShouldGenerateDataForProperty(EFlattenMaterialProperties::Roughness);
-		const bool bRenderAnisotropy = Material->IsPropertyActive(MP_Anisotropy) && FlattenMaterial->ShouldGenerateDataForProperty(EFlattenMaterialProperties::Anisotropy);
-
-
-		FMaterialData* MatSet = new FMaterialData();
-		MatSet->Material = CurrentMergeData->Material;
-		
-		MatSet->PropertySizes.Add(MP_BaseColor, FlattenMaterial->RenderSize);
-		
-		if (bRenderNormal)
-		{
-			MatSet->PropertySizes.Add(MP_Normal, FlattenMaterial->RenderSize);
-		}
-
-		if (bRenderTangent)
-		{
-			MatSet->PropertySizes.Add(MP_Tangent, FlattenMaterial->RenderSize);
-		}
-
-		if (bRenderMetallic)
-		{
-			MatSet->PropertySizes.Add(MP_Metallic, FlattenMaterial->RenderSize);
-		}
-
-		if (bRenderSpecular)
-		{
-			MatSet->PropertySizes.Add(MP_Specular, FlattenMaterial->RenderSize);
-		}
-
-		if (bRenderRoughness)
-		{
-			MatSet->PropertySizes.Add(MP_Roughness, FlattenMaterial->RenderSize);
-		}
-
-		if (bRenderAnisotropy)
-		{
-			MatSet->PropertySizes.Add(MP_Anisotropy, FlattenMaterial->RenderSize);
-		}
-		
-		if (bRenderSubSurface)
-		{
-			MatSet->PropertySizes.Add(MP_SubsurfaceColor, FlattenMaterial->RenderSize);
-		}
-
-		if (bRenderOpacity)
-		{
-			MatSet->PropertySizes.Add(MP_Opacity, FlattenMaterial->RenderSize);
-		}
-
-		if (bRenderOpacityMask)
-		{			
-			MatSet->PropertySizes.Add(MP_OpacityMask, FlattenMaterial->RenderSize);
-		}
-
-		if (bRenderEmissive)
-		{
-			MatSet->PropertySizes.Add(MP_EmissiveColor, FlattenMaterial->RenderSize);
-		}
-
-		MatSettings.Add(MatSet);
-	}
-	
-	TArray<FBakeOutput> BakeOutputs;
-	IMaterialBakingModule& Module = FModuleManager::Get().LoadModuleChecked<IMaterialBakingModule>("MaterialBaking");
-	Module.BakeMaterials(MatSettings, MeshSettings, BakeOutputs);
-	
-	for (int32 MaterialIndex = 0; MaterialIndex < NumMaterials; ++MaterialIndex)
-	{
-		FFlattenMaterial* FlattenMaterial = OutFlattenMaterials[MaterialIndex];
-		const FBakeOutput& Output = BakeOutputs[MaterialIndex];
-
-		if (Output.PropertyData.Contains(MP_BaseColor))
-		{
-			const TArray<FColor>& ColorData = Output.PropertyData.FindChecked(MP_BaseColor);
-			const FIntPoint& DataSize = Output.PropertySizes.FindChecked(MP_BaseColor);
-			FlattenMaterial->GetPropertySamples(EFlattenMaterialProperties::Diffuse) = ColorData;
-			FlattenMaterial->SetPropertySize(EFlattenMaterialProperties::Diffuse, DataSize);
-		}
-
-		if (Output.PropertyData.Contains(MP_Metallic))
-		{
-			const TArray<FColor>& ColorData = Output.PropertyData.FindChecked(MP_Metallic);
-			const FIntPoint& DataSize = Output.PropertySizes.FindChecked(MP_Metallic);
-			FlattenMaterial->GetPropertySamples(EFlattenMaterialProperties::Metallic) = ColorData;
-			FlattenMaterial->SetPropertySize(EFlattenMaterialProperties::Metallic, DataSize);
-		}
-
-		if (Output.PropertyData.Contains(MP_Specular))
-		{
-			const TArray<FColor>& ColorData = Output.PropertyData.FindChecked(MP_Specular);
-			const FIntPoint& DataSize = Output.PropertySizes.FindChecked(MP_Specular);
-			FlattenMaterial->GetPropertySamples(EFlattenMaterialProperties::Specular) = ColorData;
-			FlattenMaterial->SetPropertySize(EFlattenMaterialProperties::Specular, DataSize);
-		}
-
-		if (Output.PropertyData.Contains(MP_Roughness))
-		{
-			const TArray<FColor>& ColorData = Output.PropertyData.FindChecked(MP_Roughness);
-			const FIntPoint& DataSize = Output.PropertySizes.FindChecked(MP_Roughness);
-			FlattenMaterial->GetPropertySamples(EFlattenMaterialProperties::Roughness) = ColorData;
-			FlattenMaterial->SetPropertySize(EFlattenMaterialProperties::Roughness, DataSize);
-		}
-
-		if (Output.PropertyData.Contains(MP_Anisotropy))
-		{
-			const TArray<FColor>& ColorData = Output.PropertyData.FindChecked(MP_Anisotropy);
-			const FIntPoint& DataSize = Output.PropertySizes.FindChecked(MP_Anisotropy);
-			FlattenMaterial->GetPropertySamples(EFlattenMaterialProperties::Anisotropy) = ColorData;
-			FlattenMaterial->SetPropertySize(EFlattenMaterialProperties::Anisotropy, DataSize);
-		}
-
-		if (Output.PropertyData.Contains(MP_Normal))
-		{
-			const TArray<FColor>& ColorData = Output.PropertyData.FindChecked(MP_Normal);
-			const FIntPoint& DataSize = Output.PropertySizes.FindChecked(MP_Normal);
-			FlattenMaterial->GetPropertySamples(EFlattenMaterialProperties::Normal) = ColorData;
-			FlattenMaterial->SetPropertySize(EFlattenMaterialProperties::Normal, DataSize);
-		}
-		else
-		{
-			// Make sure we output a default normal value in case the material does not generate one (to prevent issues with combining meshes with and without normal maps being atlased together)
-			TArray<FColor>& Samples = FlattenMaterial->GetPropertySamples(EFlattenMaterialProperties::Normal);
-			Samples.Add(FColor(128, 128, 255));
-			FlattenMaterial->SetPropertySize(EFlattenMaterialProperties::Normal, FIntPoint(1, 1));
-		}
-
-		if (Output.PropertyData.Contains(MP_Tangent))
-		{
-			const TArray<FColor>& ColorData = Output.PropertyData.FindChecked(MP_Tangent);
-			const FIntPoint& DataSize = Output.PropertySizes.FindChecked(MP_Tangent);
-			FlattenMaterial->GetPropertySamples(EFlattenMaterialProperties::Tangent) = ColorData;
-			FlattenMaterial->SetPropertySize(EFlattenMaterialProperties::Tangent, DataSize);
-		}
-		else
-		{
-			// Make sure we output a default tangent value in case the material does not generate one (to prevent issues with combining meshes with and without tangent maps being atlased together)
-			TArray<FColor>& Samples = FlattenMaterial->GetPropertySamples(EFlattenMaterialProperties::Tangent);
-			Samples.Add(FColor(255, 128, 128));
-			FlattenMaterial->SetPropertySize(EFlattenMaterialProperties::Tangent, FIntPoint(1, 1));
-		}
-
-		if (Output.PropertyData.Contains(MP_Opacity))
-		{
-			const TArray<FColor>& ColorData = Output.PropertyData.FindChecked(MP_Opacity);
-			const FIntPoint& DataSize = Output.PropertySizes.FindChecked(MP_Opacity);
-			FlattenMaterial->GetPropertySamples(EFlattenMaterialProperties::Opacity) = ColorData;
-			FlattenMaterial->SetPropertySize(EFlattenMaterialProperties::Opacity, DataSize);
-		}
-
-		if (Output.PropertyData.Contains(MP_OpacityMask))
-		{
-			const TArray<FColor>& ColorData = Output.PropertyData.FindChecked(MP_OpacityMask);
-			const FIntPoint& DataSize = Output.PropertySizes.FindChecked(MP_OpacityMask);
-			FlattenMaterial->GetPropertySamples(EFlattenMaterialProperties::OpacityMask) = ColorData;
-			FlattenMaterial->SetPropertySize(EFlattenMaterialProperties::OpacityMask, DataSize);
-		}
-
-		if (Output.PropertyData.Contains(MP_EmissiveColor))
-		{
-			const TArray<FColor>& ColorData = Output.PropertyData.FindChecked(MP_EmissiveColor);
-			const FIntPoint& DataSize = Output.PropertySizes.FindChecked(MP_EmissiveColor);
-			FlattenMaterial->GetPropertySamples(EFlattenMaterialProperties::Emissive) = ColorData;
-			FlattenMaterial->SetPropertySize(EFlattenMaterialProperties::Emissive, DataSize);
-			FlattenMaterial->EmissiveScale = Output.EmissiveScale;
-		}
-
-		UE_LOG(LogMaterialUtilities, Log, TEXT("Material flattening done. (%s)"), *MergeData[MaterialIndex]->Material->GetName());
-	}
-
-	return true;
-}
-
-
 bool FMaterialUtilities::ExportMaterial(struct FMaterialMergeData& InMaterialData, FFlattenMaterial& OutFlattenMaterial, struct FExportMaterialProxyCache* ProxyCache)
 {
 	UMaterialInterface* Material = InMaterialData.Material;
diff --git a/Engine/Source/Developer/MaterialUtilities/Public/MaterialUtilities.h b/Engine/Source/Developer/MaterialUtilities/Public/MaterialUtilities.h
--- a/Engine/Source/Developer/MaterialUtilities/Public/MaterialUtilities.h
+++ b/Engine/Source/Developer/MaterialUtilities/Public/MaterialUtilities.h
@@ -266,79 +266,6 @@ public:
 	 */
 	static bool SupportsExport(EBlendMode InBlendMode, EMaterialProperty InMaterialProperty);
 
-	/**
-	 * Renders specified material property into texture
-	 *
-	 * @param InWorld				World object to use for material property rendering
-	 * @param InMaterial			Target material
-	 * @param InMaterialProperty	Material property to render
-	 * @param InRenderTarget		Render target to render to
-	 * @param OutBMP				Output array of rendered samples 
-	 * @return						Whether operation was successful
-	 */	
-	UE_DEPRECATED(4.11, "Please use ExportMaterialProperty function with new signature")
-	static bool ExportMaterialProperty(UWorld* InWorld, UMaterialInterface* InMaterial, EMaterialProperty InMaterialProperty, UTextureRenderTarget2D* InRenderTarget, TArray<FColor>& OutBMP);
-
-	/**
-	* Renders specified material property into texture
-	*
-	* @param InMaterial			Target material
-	* @param InMaterialProperty	Material property to render	
-	* @param OutBMP				Output array of rendered samples
-	* @param OutSize			Output size of the rendered samples
-	* @return					Whether operation was successful
-	*/
-	UE_DEPRECATED(4.19, "Please use new functionality in MaterialBaking module")
-	static bool ExportMaterialProperty(UMaterialInterface* InMaterial, EMaterialProperty InMaterialProperty, TArray<FColor>& OutBMP, FIntPoint& OutSize );
-
-	/**
-	* Renders specified material property into texture
-	*
-	* @param InMaterial			Target material
-	* @param InMaterialProperty	Material property to render
-	* @param InSize				Input size for the rendered samples
-	* @param OutBMP				Output array of rendered samples	
-	* @return					Whether operation was successful
-	*/
-	UE_DEPRECATED(4.19, "Please use new functionality in MaterialBaking module")
-	static bool ExportMaterialProperty(UMaterialInterface* InMaterial, EMaterialProperty InMaterialProperty, FIntPoint InSize, TArray<FColor>& OutBMP );
-
-	/**
-	 * Renders specified material property into texture
-	 *
-	 * @param InWorld				World object to use for material property rendering
-	 * @param InMaterial			Target material
-	 * @param InMaterialProperty	Material property to render
-	 * @param OutBMP				Output array of rendered samples 
-	 * @return						Whether operation was successful
-	 */
-	UE_DEPRECATED(4.11, "Please use ExportMaterialProperty function with new signature")
-	static bool ExportMaterialProperty(UWorld* InWorld, UMaterialInterface* InMaterial, EMaterialProperty InMaterialProperty, FIntPoint& OutSize, TArray<FColor>& OutBMP);
-
-	/**
-	 * Flattens specified material
-	 *
-	 * @param InWorld				World object to use for material rendering
-	 * @param InMaterial			Target material
-	 * @param OutFlattenMaterial	Output flattened material
-	 * @return						Whether operation was successful
-	 */
-	UE_DEPRECATED(4.11, "Please use ExportMaterial function with new signature")
-	static bool ExportMaterial(UWorld* InWorld, UMaterialInterface* InMaterial, FFlattenMaterial& OutFlattenMaterial);
-	
-	/**
-	* Flattens specified material
-	*
-	* @param InMaterial				Target material
-	* @param OutFlattenMaterial		Output flattened material
-	* @return						Whether operation was successful
-	*/
-	UE_DEPRECATED(4.19, "Please use new functionality in MaterialBaking module")
-	static bool ExportMaterial(UMaterialInterface* InMaterial, FFlattenMaterial& OutFlattenMaterial, struct FExportMaterialProxyCache* ProxyCache = nullptr);
-
-	UE_DEPRECATED(4.17, "Please use new functionality in MaterialBaking module")
-	static bool ExportMaterials(TArray<FMaterialMergeData*>& MergeData, TArray<FFlattenMaterial*>& OutFlattenMaterials);
-
 	/**
 	 * Flattens specified landscape material
 	 *
diff --git a/Engine/Source/Developer/MeshUtilities/Private/MeshUtilities.cpp b/Engine/Source/Developer/MeshUtilities/Private/MeshUtilities.cpp
--- a/Engine/Source/Developer/MeshUtilities/Private/MeshUtilities.cpp
+++ b/Engine/Source/Developer/MeshUtilities/Private/MeshUtilities.cpp
@@ -6403,101 +6403,6 @@ void FMeshUtilities::ConvertActorMeshesToStaticMeshUIAction(const TArray<AActor*
 	}
 }
 
-/************************************************************************/
-/*  DEPRECATED FUNCTIONALITY                                            */
-/************************************************************************/
-IMeshReduction* FMeshUtilities::GetStaticMeshReductionInterface()
-{
-	IMeshReductionManagerModule& Module = FModuleManager::Get().LoadModuleChecked<IMeshReductionManagerModule>("MeshReductionInterface");
-	return Module.GetStaticMeshReductionInterface();
-}
-
-IMeshReduction* FMeshUtilities::GetSkeletalMeshReductionInterface()
-{
-	IMeshReductionManagerModule& Module = FModuleManager::Get().LoadModuleChecked<IMeshReductionManagerModule>("MeshReductionInterface");
-	return Module.GetSkeletalMeshReductionInterface();
-}
-
-IMeshMerging* FMeshUtilities::GetMeshMergingInterface()
-{
-	IMeshReductionManagerModule& Module = FModuleManager::Get().LoadModuleChecked<IMeshReductionManagerModule>("MeshReductionInterface");
-	return Module.GetMeshMergingInterface();
-}
-
-void FMeshUtilities::MergeActors(
-	const TArray<AActor*>& SourceActors,
-	const FMeshMergingSettings& InSettings,
-	UPackage* InOuter,
-	const FString& InBasePackageName,
-	TArray<UObject*>& OutAssetsToSync,
-	FVector& OutMergedActorLocation,
-	bool bSilent) const
-{
-	checkf(SourceActors.Num(), TEXT("No actors supplied for merging"));
-	
-	// Collect all primitive components
-	TInlineComponentArray<UPrimitiveComponent*> PrimComps;
-	for (AActor* Actor : SourceActors)
-	{
-		Actor->GetComponents<UPrimitiveComponent>(PrimComps);
-	}
-
-	// Filter only components we want (static mesh and shape)
-	TArray<UPrimitiveComponent*> ComponentsToMerge;
-	for (UPrimitiveComponent* PrimComponent : PrimComps)
-	{
-		UStaticMeshComponent* MeshComponent = Cast<UStaticMeshComponent>(PrimComponent);
-		if (MeshComponent && 
-			MeshComponent->GetStaticMesh() != nullptr &&
-			MeshComponent->GetStaticMesh()->GetNumSourceModels() > 0)
-		{
-			ComponentsToMerge.Add(MeshComponent);
-		}
-
-		UShapeComponent* ShapeComponent = Cast<UShapeComponent>(PrimComponent);
-		if (ShapeComponent)
-		{
-			ComponentsToMerge.Add(ShapeComponent);
-		}
-	}
-
-	checkf(SourceActors.Num(), TEXT("No valid components found in actors supplied for merging"));
-
-	UWorld* World = SourceActors[0]->GetWorld();
-	checkf(World != nullptr, TEXT("Invalid world retrieved from Actor"));
-	const float ScreenSize = TNumericLimits<float>::Max();
-
-	const IMeshMergeUtilities& Module = FModuleManager::Get().LoadModuleChecked<IMeshMergeModule>("MeshMergeUtilities").GetUtilities();
-	Module.MergeComponentsToStaticMesh(ComponentsToMerge, World, InSettings, nullptr, InOuter, InBasePackageName, OutAssetsToSync, OutMergedActorLocation, ScreenSize, bSilent);
-}
-
-void FMeshUtilities::MergeStaticMeshComponents(
-	const TArray<UStaticMeshComponent*>& ComponentsToMerge,
-	UWorld* World,
-	const FMeshMergingSettings& InSettings,
-	UPackage* InOuter,
-	const FString& InBasePackageName,
-	TArray<UObject*>& OutAssetsToSync,
-	FVector& OutMergedActorLocation,
-	const float ScreenSize,
-	bool bSilent /*= false*/) const
-{
-	const IMeshMergeUtilities& Module = FModuleManager::Get().LoadModuleChecked<IMeshMergeModule>("MeshMergeUtilities").GetUtilities();
-
-	// Convert array of StaticMeshComponents to PrimitiveComponents
-	TArray<UPrimitiveComponent*> PrimCompsToMerge;
-	Algo::Transform(ComponentsToMerge, PrimCompsToMerge, [](UStaticMeshComponent* StaticMeshComp) { return StaticMeshComp; });
-
-	Module.MergeComponentsToStaticMesh(PrimCompsToMerge, World, InSettings, nullptr, InOuter, InBasePackageName, OutAssetsToSync, OutMergedActorLocation, ScreenSize, bSilent);
-}
-
-void FMeshUtilities::CreateProxyMesh(const TArray<AActor*>& InActors, const struct FMeshProxySettings& InMeshProxySettings, UPackage* InOuter, const FString& InProxyBasePackageName, const FGuid InGuid, FCreateProxyDelegate InProxyCreatedDelegate, const bool bAllowAsync,
-	const float ScreenAreaSize /*= 1.0f*/)
-{
-	const IMeshMergeUtilities& Module = FModuleManager::Get().LoadModuleChecked<IMeshMergeModule>("MeshMergeUtilities").GetUtilities();
-	Module.CreateProxyMesh(InActors, InMeshProxySettings, InOuter, InProxyBasePackageName, InGuid, InProxyCreatedDelegate, bAllowAsync, ScreenAreaSize);
-}
-
 bool FMeshUtilities::GenerateUniqueUVsForStaticMesh(const FRawMesh& RawMesh, int32 TextureResolution, bool bMergeIdenticalMaterials, TArray<FVector2f>& OutTexCoords) const
 {
 	// Create a copy of original mesh (only copy necessary data)
@@ -6605,9 +6510,4 @@ bool FMeshUtilities::GenerateUniqueUVsForStaticMesh(const FRawMesh& RawMesh, int
 	return GenerateUniqueUVsForStaticMesh(RawMesh, TextureResolution, false, OutTexCoords);
 }
 
-void FMeshUtilities::FlattenMaterialsWithMeshData(TArray<UMaterialInterface*>& InMaterials, TArray<FRawMeshExt>& InSourceMeshes, TMap<FMeshIdAndLOD, TArray<int32>>& InMaterialIndexMap, TArray<bool>& InMeshShouldBakeVertexData, const FMaterialProxySettings &InMaterialProxySettings, TArray<FFlattenMaterial> &OutFlattenedMaterials) const
-{
-	checkf(false, TEXT("Function is removed, use functionality in new MeshMergeUtilities Module"));
-}
-
 #undef LOCTEXT_NAMESPACE
diff --git a/Engine/Source/Developer/MeshUtilities/Private/MeshUtilitiesPrivate.h b/Engine/Source/Developer/MeshUtilities/Private/MeshUtilitiesPrivate.h
--- a/Engine/Source/Developer/MeshUtilities/Private/MeshUtilitiesPrivate.h
+++ b/Engine/Source/Developer/MeshUtilities/Private/MeshUtilitiesPrivate.h
@@ -17,45 +17,6 @@
 
 class FMeshUtilities : public IMeshUtilities
 {
-public:
-	UE_DEPRECATED(4.17, "Use functionality in new MeshReduction Module")
-	virtual IMeshReduction* GetStaticMeshReductionInterface() override;
-	
-	UE_DEPRECATED(4.17, "Use functionality in new MeshReduction Module")
-	virtual IMeshReduction* GetSkeletalMeshReductionInterface() override;
-	
-	UE_DEPRECATED(4.17, "Use functionality in new MeshReduction Module")
-	virtual IMeshMerging* GetMeshMergingInterface() override;
-	
-	UE_DEPRECATED(4.17, "Use functionality in new MeshMergeUtilities Module")
-	virtual void MergeActors(
-		const TArray<AActor*>& SourceActors,
-		const FMeshMergingSettings& InSettings,
-		UPackage* InOuter,
-		const FString& InBasePackageName,
-		TArray<UObject*>& OutAssetsToSync,
-		FVector& OutMergedActorLocation,
-		bool bSilent = false) const override;
-
-	UE_DEPRECATED(4.17, "Use functionality in new MeshMergeUtilities Module")
-	virtual void MergeStaticMeshComponents(
-		const TArray<UStaticMeshComponent*>& ComponentsToMerge,
-		UWorld* World,
-		const FMeshMergingSettings& InSettings,
-		UPackage* InOuter,
-		const FString& InBasePackageName,
-		TArray<UObject*>& OutAssetsToSync,
-		FVector& OutMergedActorLocation,
-		const float ScreenSize,
-		bool bSilent = false) const override;
-
-	UE_DEPRECATED(4.17, "Use functionality in new MeshMergeUtilities Module")
-	virtual void CreateProxyMesh(const TArray<AActor*>& InActors, const struct FMeshProxySettings& InMeshProxySettings, UPackage* InOuter, const FString& InProxyBasePackageName, const FGuid InGuid, FCreateProxyDelegate InProxyCreatedDelegate, const bool bAllowAsync,
-	const float ScreenAreaSize = 1.0f) override;
-
-	UE_DEPRECATED(4.17, "Function is removed, use functionality in new MeshMergeUtilities Module")
-	virtual void FlattenMaterialsWithMeshData(TArray<UMaterialInterface*>& InMaterials, TArray<FRawMeshExt>& InSourceMeshes, TMap<FMeshIdAndLOD, TArray<int32>>& InMaterialIndexMap, TArray<bool>& InMeshShouldBakeVertexData, const FMaterialProxySettings &InMaterialProxySettings, TArray<FFlattenMaterial> &OutFlattenedMaterials) const override;
-
 private:
 	/** Cached version string. */
 	FString VersionString;
diff --git a/Engine/Source/Developer/MeshUtilities/Public/MeshUtilities.h b/Engine/Source/Developer/MeshUtilities/Public/MeshUtilities.h
--- a/Engine/Source/Developer/MeshUtilities/Public/MeshUtilities.h
+++ b/Engine/Source/Developer/MeshUtilities/Public/MeshUtilities.h
@@ -60,84 +60,7 @@ struct FOverlappingCorners;
 class IMeshUtilities : public IModuleInterface
 {
 public:
-	/************************************************************************/
-	/*  DEPRECATED FUNCTIONALITY                                            */
-	/************************************************************************/
 
-	/**
-	* Harvest static mesh components from input actors
-	* and merge into signle mesh grouping them by unique materials
-	*
-	* @param SourceActors				List of actors to merge
-	* @param InSettings				Settings to use
-	* @param InOuter					Outer if required
-	* @param InBasePackageName			Destination package name for a generated assets. Used if Outer is null.
-	* @param UseLOD					-1 if you'd like to build for all LODs. If you specify, that LOD mesh for source meshes will be used to merge the mesh
-	*									This is used by hierarchical building LODs
-	* @param OutAssetsToSync			Merged mesh assets
-	* @param OutMergedActorLocation	World position of merged mesh
-	*/
-
-	virtual void MergeActors(
-		const TArray<AActor*>& SourceActors,
-		const FMeshMergingSettings& InSettings,
-		UPackage* InOuter,
-		const FString& InBasePackageName,
-		TArray<UObject*>& OutAssetsToSync, 
-		FVector& OutMergedActorLocation, 
-		bool bSilent=false) const = 0;
-	/**
-	* MergeStaticMeshComponents
-	*
-	* @param ComponentsToMerge - Components to merge
-	* @param World - World in which the component reside
-	* @param InSettings	- Settings to use
-	* @param InOuter - Outer if required
-	* @param InBasePackageName - Destination package name for a generated assets. Used if Outer is null.
-	* @param UseLOD	-1 if you'd like to build for all LODs. If you specify, that LOD mesh for source meshes will be used to merge the mesh
-	*									This is used by hierarchical building LODs
-	* @param OutAssetsToSync Merged mesh assets
-	* @param OutMergedActorLocation	World position of merged mesh
-	* @param ViewDistance Distance for LOD determination
-	* @param bSilent Non-verbose flag
-	* @return void
-	*/
-	virtual void MergeStaticMeshComponents(
-		const TArray<UStaticMeshComponent*>& ComponentsToMerge,
-		UWorld* World,
-		const FMeshMergingSettings& InSettings,
-		UPackage* InOuter,
-		const FString& InBasePackageName,
-		TArray<UObject*>& OutAssetsToSync,
-		FVector& OutMergedActorLocation,
-		const float ScreenAreaSize,
-		bool bSilent /*= false*/) const = 0;
-	
-	/**
-	* Creates a (proxy)-mesh combining the static mesh components from the given list of actors (at the moment this requires having Simplygon)
-	*
-	* @param InActors - List of Actors to merge
-	* @param InMeshProxySettings - Merge settings
-	* @param InOuter - Package for a generated assets, if NULL new packages will be created for each asset
-	* @param InProxyBasePackageName - Will be used for naming generated assets, in case InOuter is not specified ProxyBasePackageName will be used as long package name for creating new packages
-	* @param InGuid - Guid identifying the data used for this proxy job
-	* @param InProxyCreatedDelegate - Delegate callback for when the proxy is finished
-	* @param bAllowAsync - Flag whether or not this call could be run async (SimplygonSwarm)
-	*/
-	virtual void CreateProxyMesh(const TArray<class AActor*>& InActors, const struct FMeshProxySettings& InMeshProxySettings, UPackage* InOuter, const FString& InProxyBasePackageName, const FGuid InGuid, FCreateProxyDelegate InProxyCreatedDelegate, const bool bAllowAsync = false, const float ScreenAreaSize = 1.0f) = 0;
-
-	/**
-	* FlattenMaterialsWithMeshData
-	*
-	* @param InMaterials - List of unique materials used by InSourceMeshes
-	* @param InSourceMeshes - List of raw meshes used to flatten the materials with (vertex data)
-	* @param InMaterialIndexMap - Map used for mapping the raw meshes to the correct materials
-	* @param InMeshShouldBakeVertexData - Array of flags to determine whether or not a mesh requires to have its vertex data baked down
-	* @param InMaterialProxySettings - Settings for creating the flattened material
-	* @param OutFlattenedMaterials - List of flattened materials (one for each mesh)
-	*/
-	virtual	void FlattenMaterialsWithMeshData(TArray<UMaterialInterface*>& InMaterials, TArray<struct FRawMeshExt>& InSourceMeshes, TMap<FMeshIdAndLOD, TArray<int32>>& InMaterialIndexMap, TArray<bool>& InMeshShouldBakeVertexData, const FMaterialProxySettings &InMaterialProxySettings, TArray<FFlattenMaterial> &OutFlattenedMaterials) const = 0;
-	
 	/**
 	* Calculates (new) non-overlapping UV coordinates for the given Raw Mesh
 	*
@@ -150,14 +73,6 @@ public:
 	virtual bool GenerateUniqueUVsForStaticMesh(const FRawMesh& RawMesh, int32 TextureResolution, TArray<FVector2f>& OutTexCoords) const = 0;
 	virtual bool GenerateUniqueUVsForStaticMesh(const FRawMesh& RawMesh, int32 TextureResolution, bool bMergeIdenticalMaterials, TArray<FVector2f>& OutTexCoords) const = 0;
 	
-	/** Returns the mesh reduction plugin if available. */
-	virtual IMeshReduction* GetStaticMeshReductionInterface() = 0;
-
-	/** Returns the mesh reduction plugin if available. */
-	virtual IMeshReduction* GetSkeletalMeshReductionInterface() = 0;
-
-	/** Returns the mesh merging plugin if available. */
-	virtual IMeshMerging* GetMeshMergingInterface() = 0;
 public:
 	/** Returns a string uniquely identifying this version of mesh utilities. */
 	virtual const FString& GetVersionString() const = 0;
diff --git a/Engine/Source/Developer/TargetDeviceServices/Public/ITargetDeviceProxy.h b/Engine/Source/Developer/TargetDeviceServices/Public/ITargetDeviceProxy.h
--- a/Engine/Source/Developer/TargetDeviceServices/Public/ITargetDeviceProxy.h
+++ b/Engine/Source/Developer/TargetDeviceServices/Public/ITargetDeviceProxy.h
@@ -13,16 +13,6 @@ class ITargetDeviceProxy;
 
 struct FGuid;
 
-
-/** Type definition for shared pointers to instances of ITargetDeviceProxy. */
-UE_DEPRECATED(4.16, "ITargetDeviceProxyPtr is deprecated. Please use 'TSharedPtr<ITargetDeviceProxy>' instead!")
-typedef TSharedPtr<class ITargetDeviceProxy> ITargetDeviceProxyPtr;
-
-/** Type definition for shared references to instances of ITargetDeviceProxy. */
-UE_DEPRECATED(4.16, "ITargetDeviceProxyRef is deprecated. Please use 'TSharedPtr<ITargetDeviceProxy>' instead!")
-typedef TSharedRef<class ITargetDeviceProxy> ITargetDeviceProxyRef;
-
-
 /**
  * Interface for target device proxies.
  */
diff --git a/Engine/Source/Developer/TargetDeviceServices/Public/ITargetDeviceProxyManager.h b/Engine/Source/Developer/TargetDeviceServices/Public/ITargetDeviceProxyManager.h
--- a/Engine/Source/Developer/TargetDeviceServices/Public/ITargetDeviceProxyManager.h
+++ b/Engine/Source/Developer/TargetDeviceServices/Public/ITargetDeviceProxyManager.h
@@ -12,16 +12,6 @@ class ITargetDeviceProxy;
 class ITargetDeviceProxyManager;
 class ITargetDeviceServiceManager;
 
-
-/** Type definition for shared pointers to instances of ITargetDeviceProxyManager. */
-UE_DEPRECATED(4.16, "ITargetDeviceProxyManagerPtr is deprecated. Please use 'TSharedPtr<ITargetDeviceProxyManager>' instead!")
-typedef TSharedPtr<class ITargetDeviceProxyManager> ITargetDeviceProxyManagerPtr;
-
-/** Type definition for shared references to instances of ITargetDeviceProxyManager. */
-UE_DEPRECATED(4.16, "ITargetDeviceProxyManagerRef is deprecated. Please use 'TSharedPtr<ITargetDeviceProxyManager>' instead!")
-typedef TSharedRef<class ITargetDeviceProxyManager> ITargetDeviceProxyManagerRef;
-
-
 /**
  * Interface for device proxy managers.
  */
diff --git a/Engine/Source/Developer/TargetDeviceServices/Public/ITargetDeviceServiceManager.h b/Engine/Source/Developer/TargetDeviceServices/Public/ITargetDeviceServiceManager.h
--- a/Engine/Source/Developer/TargetDeviceServices/Public/ITargetDeviceServiceManager.h
+++ b/Engine/Source/Developer/TargetDeviceServices/Public/ITargetDeviceServiceManager.h
@@ -11,16 +11,6 @@
 class FString;
 class ITargetDeviceService;
 
-
-/** Type definition for shared pointers to instances of ITargetDeviceServiceManager. */
-UE_DEPRECATED(4.16, "ITargetDeviceServiceManagerPtr is deprecated. Please use 'TSharedPtr<ITargetDeviceServiceManager>' instead!")
-typedef TSharedPtr<class ITargetDeviceServiceManager> ITargetDeviceServiceManagerPtr;
-
-/** Type definition for shared references to instances of ITargetDeviceServiceManager. */
-UE_DEPRECATED(4.16, "ITargetDeviceServiceManagerRef is deprecated. Please use 'TSharedPtr<ITargetDeviceServiceManager>' instead!")
-typedef TSharedRef<class ITargetDeviceServiceManager> ITargetDeviceServiceManagerRef;
-
-
 /**
  * Interface for target device service managers.
  */
diff --git a/Engine/Source/Editor/AnimGraph/Public/AnimGraphNode_Base.h b/Engine/Source/Editor/AnimGraph/Public/AnimGraphNode_Base.h
--- a/Engine/Source/Editor/AnimGraph/Public/AnimGraphNode_Base.h
+++ b/Engine/Source/Editor/AnimGraph/Public/AnimGraphNode_Base.h
@@ -340,8 +340,6 @@ public:
 	virtual void DrawCanvas(FViewport& InViewport, FSceneView& View, FCanvas& Canvas, USkeletalMeshComponent * PreviewSkelMeshComp) const {}
 	// Function to collect strings from nodes to display in the viewport.
 	// Use this rather than DrawCanvas when adding general text to the viewport.
-	UE_DEPRECATED(4.16, "Please use GetOnScreenDebugInfo(TArray<FText>& DebugInfo, FAnimNode_Base* RuntimeAnimNode, USkeletalMeshComponent* PreviewSkelMeshComp)")
-	virtual void GetOnScreenDebugInfo(TArray<FText>& DebugInfo, USkeletalMeshComponent* PreviewSkelMeshComp) const {}
 	virtual void GetOnScreenDebugInfo(TArray<FText>& DebugInfo, FAnimNode_Base* RuntimeAnimNode, USkeletalMeshComponent* PreviewSkelMeshComp) const {}
 
 	/** Called after editing a default value to update internal node from pin defaults. This is needed for forwarding code to propagate values to preview. */
diff --git a/Engine/Source/Editor/BlueprintGraph/Private/K2Node.cpp b/Engine/Source/Editor/BlueprintGraph/Private/K2Node.cpp
--- a/Engine/Source/Editor/BlueprintGraph/Private/K2Node.cpp
+++ b/Engine/Source/Editor/BlueprintGraph/Private/K2Node.cpp
@@ -55,7 +55,6 @@ namespace UK2Node_Private
 UK2Node::UK2Node(const FObjectInitializer& ObjectInitializer)
 	: Super(ObjectInitializer)
 {
-	bAllowSplitPins_DEPRECATED = true;
 	OrphanedPinSaveMode = ESaveOrphanPinMode::SaveAllButExec;
 }
 
@@ -1292,14 +1291,7 @@ void UK2Node::DestroyPinList(TArray<UEdGraphPin*>& InPins)
 
 bool UK2Node::CanSplitPin(const UEdGraphPin* Pin) const
 {
-	PRAGMA_DISABLE_DEPRECATION_WARNINGS
-	// AllowSplitPins is deprecated. Remove this block when that function is eventually removed.
-	if (AllowSplitPins())
-	{
-		return (Pin->GetOwningNode() == this && !Pin->bNotConnectable && Pin->LinkedTo.Num() == 0 && Pin->PinType.PinCategory == UEdGraphSchema_K2::PC_Struct);
-	}
-	PRAGMA_ENABLE_DEPRECATION_WARNINGS
-	return false;
+	return (Pin->GetOwningNode() == this && !Pin->bNotConnectable && Pin->LinkedTo.Num() == 0 && Pin->PinType.PinCategory == UEdGraphSchema_K2::PC_Struct);
 }
 
 UK2Node* UK2Node::ExpandSplitPin(FKismetCompilerContext* CompilerContext, UEdGraph* SourceGraph, UEdGraphPin* Pin)
diff --git a/Engine/Source/Editor/BlueprintGraph/Private/K2Node_VariableGet.cpp b/Engine/Source/Editor/BlueprintGraph/Private/K2Node_VariableGet.cpp
--- a/Engine/Source/Editor/BlueprintGraph/Private/K2Node_VariableGet.cpp
+++ b/Engine/Source/Editor/BlueprintGraph/Private/K2Node_VariableGet.cpp
@@ -427,13 +427,11 @@ void UK2Node_VariableGet::ValidateNodeDuringCompilation(FCompilerResultsLog& Mes
 
 				if (PropertyReadableState == FBlueprintEditorUtils::EPropertyReadableState::NotBlueprintVisible)
 				{
-					// UE_DEPRECATED(4.17) ... make this an error
-					MessageLog.Warning(*FText::Format(LOCTEXT("UnableToGet_NotVisible", "{VariableName} is not blueprint visible (BlueprintReadOnly or BlueprintReadWrite). Please fix mark up or cease accessing as this will be made an error in a future release. @@"), Args).ToString(), this);
+					MessageLog.Error(*FText::Format(LOCTEXT("UnableToGet_NotVisible", "{VariableName} is not blueprint visible (BlueprintReadOnly or BlueprintReadWrite). Please fix mark up or cease accessing as this will be made an error in a future release. @@"), Args).ToString(), this);
 				}
 				else if (PropertyReadableState == FBlueprintEditorUtils::EPropertyReadableState::Private)
 				{
-					// UE_DEPRECATED(4.17) ... make this an error
-					MessageLog.Warning(*FText::Format(LOCTEXT("UnableToGet_ReadOnly", "{VariableName} is private and not accessible in this context. Please fix mark up or cease accessing as this will be an error in a future release. @@"), Args).ToString(), this);
+					MessageLog.Error(*FText::Format(LOCTEXT("UnableToGet_ReadOnly", "{VariableName} is private and not accessible in this context. Please fix mark up or cease accessing as this will be an error in a future release. @@"), Args).ToString(), this);
 				}
 				else
 				{
diff --git a/Engine/Source/Editor/KismetCompiler/Private/KismetCompiler.cpp b/Engine/Source/Editor/KismetCompiler/Private/KismetCompiler.cpp
--- a/Engine/Source/Editor/KismetCompiler/Private/KismetCompiler.cpp
+++ b/Engine/Source/Editor/KismetCompiler/Private/KismetCompiler.cpp
@@ -2924,13 +2924,6 @@ FPinConnectionResponse FKismetCompilerContext::CopyPinLinksToIntermediate(UEdGra
 	return ConnectionResult;
 }
 
-UK2Node_TemporaryVariable* FKismetCompilerContext::SpawnInternalVariable(UEdGraphNode* SourceNode, const FString& Category, const FString& SubCategory, UObject* SubcategoryObject, bool bIsArray, bool bIsSet, bool bIsMap, const FEdGraphTerminalType& ValueTerminalType)
-{
-	PRAGMA_DISABLE_DEPRECATION_WARNINGS
-	return SpawnInternalVariable(SourceNode, Category, SubCategory, SubcategoryObject, FEdGraphPinType::ToPinContainerType(bIsArray, bIsSet, bIsMap), ValueTerminalType);
-	PRAGMA_ENABLE_DEPRECATION_WARNINGS
-}
-
 UK2Node_TemporaryVariable* FKismetCompilerContext::SpawnInternalVariable(UEdGraphNode* SourceNode, const FName Category, const FName SubCategory, UObject* SubcategoryObject, EPinContainerType PinContainerType, const FEdGraphTerminalType& ValueTerminalType)
 {
 	UK2Node_TemporaryVariable* Result = SpawnIntermediateNode<UK2Node_TemporaryVariable>(SourceNode);
diff --git a/Engine/Source/Editor/KismetCompiler/Public/KismetCompiler.h b/Engine/Source/Editor/KismetCompiler/Public/KismetCompiler.h
--- a/Engine/Source/Editor/KismetCompiler/Public/KismetCompiler.h
+++ b/Engine/Source/Editor/KismetCompiler/Public/KismetCompiler.h
@@ -255,21 +255,6 @@ public:
 		FName NameParameter;
 	};
 
-	UE_DEPRECATED(4.17, "Use version that takes Category and SubCategory as FName, and PinContainerType instead of separate booleans for array, set, and map")
-	UK2Node_TemporaryVariable* SpawnInternalVariable(UEdGraphNode* SourceNode, const FString& Category, const FString& SubCategory, UObject* SubcategoryObject, bool bIsArray, bool bIsSet = false, bool bIsMap = false, const FEdGraphTerminalType& ValueTerminalType = FEdGraphTerminalType());
-
-	UE_DEPRECATED(4.18, "Use version that takes Category and SubCategory as FName")
-	UK2Node_TemporaryVariable* SpawnInternalVariable(UEdGraphNode* SourceNode, const FNameParameterHelper Category, const FString& SubCategory, UObject* SubcategoryObject = nullptr, EPinContainerType PinContainerType = EPinContainerType::None, const FEdGraphTerminalType& ValueTerminalType = FEdGraphTerminalType())
-	{
-		return SpawnInternalVariable(SourceNode, *Category, FName(*SubCategory), SubcategoryObject, PinContainerType, ValueTerminalType);
-	}
-
-	//UE_DEPRECATED(4.18, "Remove when removing versions that take subcategory as FString. Required to avoid ambiguity")
-	UK2Node_TemporaryVariable* SpawnInternalVariable(UEdGraphNode* SourceNode, const FNameParameterHelper Category, const TCHAR* SubCategory, UObject* SubcategoryObject = nullptr, EPinContainerType PinContainerType = EPinContainerType::None, const FEdGraphTerminalType& ValueTerminalType = FEdGraphTerminalType())
-	{
-		return SpawnInternalVariable(SourceNode, *Category, FName(SubCategory), SubcategoryObject, PinContainerType, ValueTerminalType);
-	}
-
 	UK2Node_TemporaryVariable* SpawnInternalVariable(UEdGraphNode* SourceNode, FName Category, FName SubCategory = NAME_None, UObject* SubcategoryObject = nullptr, EPinContainerType PinContainerType = EPinContainerType::None, const FEdGraphTerminalType& ValueTerminalType = FEdGraphTerminalType());
 
 	bool UsePersistentUberGraphFrame() const;
diff --git a/Engine/Source/Editor/PropertyEditor/Private/PropertyEditorModule.cpp b/Engine/Source/Editor/PropertyEditor/Private/PropertyEditorModule.cpp
--- a/Engine/Source/Editor/PropertyEditor/Private/PropertyEditorModule.cpp
+++ b/Engine/Source/Editor/PropertyEditor/Private/PropertyEditorModule.cpp
@@ -475,31 +475,6 @@ void FPropertyEditorModule::RegisterCustomPropertyTypeLayout( FName PropertyType
 	}
 }
 
-void FPropertyEditorModule::RegisterCustomPropertyTypeLayout(FName PropertyTypeName, FOnGetPropertyTypeCustomizationInstance PropertyTypeLayoutDelegate, TSharedPtr<IPropertyTypeIdentifier> Identifier, TSharedPtr<IDetailsView> ForSpecificInstance)
-{
-	if (ForSpecificInstance.IsValid())
-	{
-		ForSpecificInstance->RegisterInstancedCustomPropertyTypeLayout(PropertyTypeName, PropertyTypeLayoutDelegate, Identifier);
-	}
-	else
-	{
-		RegisterCustomPropertyTypeLayout(PropertyTypeName, PropertyTypeLayoutDelegate, Identifier);
-	}
-}
-
-void FPropertyEditorModule::UnregisterCustomPropertyTypeLayout(FName PropertyTypeName, TSharedPtr<IPropertyTypeIdentifier> InIdentifier, TSharedPtr<IDetailsView> ForSpecificInstance)
-{
-	if (ForSpecificInstance.IsValid())
-	{
-		ForSpecificInstance->UnregisterInstancedCustomPropertyTypeLayout(PropertyTypeName, InIdentifier);
-	}
-	else
-	{
-		UnregisterCustomPropertyTypeLayout(PropertyTypeName, InIdentifier);
-	}
-}
-
-
 void FPropertyEditorModule::UnregisterCustomPropertyTypeLayout( FName PropertyTypeName, TSharedPtr<IPropertyTypeIdentifier> Identifier)
 {
 	if (!PropertyTypeName.IsValid() || (PropertyTypeName == NAME_None))
diff --git a/Engine/Source/Editor/PropertyEditor/Public/DetailCategoryBuilder.h b/Engine/Source/Editor/PropertyEditor/Public/DetailCategoryBuilder.h
--- a/Engine/Source/Editor/PropertyEditor/Public/DetailCategoryBuilder.h
+++ b/Engine/Source/Editor/PropertyEditor/Public/DetailCategoryBuilder.h
@@ -201,22 +201,4 @@ public:
 	* This is designed to be used for dynamic display of advanced properties.
 	*/
 	virtual void SetShowAdvanced(bool bShowAdvanced) = 0;
-
-	UE_DEPRECATED(4.17, "AddExternalProperty is deprecated.  Use AddExternalObjectProperty instead")
-	IDetailPropertyRow* AddExternalProperty(const TArray<UObject*>& Objects, FName PropertyName, EPropertyLocation::Type Location = EPropertyLocation::Default)
-	{
-		return AddExternalObjectProperty(Objects, PropertyName, Location);
-	}
-
-	UE_DEPRECATED(4.17, "AddExternalProperty is deprecated.  Use AddExternalStructureProperty instead")
-	IDetailPropertyRow* AddExternalProperty(TSharedPtr<FStructOnScope> StructData, FName PropertyName, EPropertyLocation::Type Location = EPropertyLocation::Default)
-	{
-		return AddExternalStructureProperty(StructData, PropertyName, Location);
-	}
-
-	UE_DEPRECATED(4.17, "AddExternalProperty is deprecated.  Use AddAllExternalStructureProperties instead")
-	TArray<TSharedPtr<IPropertyHandle>> AddExternalProperties(TSharedRef<FStructOnScope> StructData, EPropertyLocation::Type Location = EPropertyLocation::Default)
-	{
-		return AddAllExternalStructureProperties(StructData, Location);
-	}
 };
diff --git a/Engine/Source/Editor/PropertyEditor/Public/IDetailChildrenBuilder.h b/Engine/Source/Editor/PropertyEditor/Public/IDetailChildrenBuilder.h
--- a/Engine/Source/Editor/PropertyEditor/Public/IDetailChildrenBuilder.h
+++ b/Engine/Source/Editor/PropertyEditor/Public/IDetailChildrenBuilder.h
@@ -138,41 +138,4 @@ public:
 	{
 		return AddExternalStructureProperty(ChildStructure, PropertyName, FAddPropertyParams().UniqueId(UniqueIdName));
 	}
-
-	UE_DEPRECATED(4.17, "AddChildCustomBuilder has been deprecated.  Use AddCustomBuilder instead")
-	IDetailChildrenBuilder& AddChildCustomBuilder(TSharedRef<IDetailCustomNodeBuilder> InCustomBuilder)
-	{
-		return AddCustomBuilder(InCustomBuilder);
-	}
-
-	UE_DEPRECATED(4.17, "AddChildGroup has been deprecated.  Use AddGroup instead")
-	class IDetailGroup& AddChildGroup(FName GroupName, const FText& LocalizedDisplayName)
-	{
-		return AddGroup(GroupName, LocalizedDisplayName);
-	}
-
-	UE_DEPRECATED(4.17, "AddChildContent has been deprecated.  Use AddCustomRow instead ")
-	class FDetailWidgetRow& AddChildContent(const FText& SearchString)
-	{
-		return AddCustomRow(SearchString);
-	}
-
-	UE_DEPRECATED(4.17, "AddChildProperty has been deprecated.  Use AddProperty instead ")
-	class IDetailPropertyRow& AddChildProperty(TSharedRef<IPropertyHandle> PropertyHandle)
-	{
-		return AddProperty(PropertyHandle);
-	}
-
-	UE_DEPRECATED(4.17, "AddChildStructure has been deprecated.  Use AddExternalStructureProperty instead ")
-	class IDetailPropertyRow* AddChildStructure(TSharedRef<FStructOnScope> ChildStructure, FName PropertyName, FName UniqueIdName)
-	{
-		return AddExternalStructureProperty(ChildStructure, PropertyName, FAddPropertyParams().UniqueId(UniqueIdName));
-	}
-
-	UE_DEPRECATED(4.17, "This version of AddStructure has been deprecated.  Use AddAllExternalStructureProperties instead ")
-	virtual TArray<TSharedPtr<IPropertyHandle>> AddStructure(TSharedRef<FStructOnScope> ChildStructure)
-	{
-		return AddAllExternalStructureProperties(ChildStructure);
-	}
-
 };
\ No newline at end of file
diff --git a/Engine/Source/Editor/PropertyEditor/Public/PropertyEditorModule.h b/Engine/Source/Editor/PropertyEditor/Public/PropertyEditorModule.h
--- a/Engine/Source/Editor/PropertyEditor/Public/PropertyEditorModule.h
+++ b/Engine/Source/Editor/PropertyEditor/Public/PropertyEditorModule.h
@@ -265,12 +265,6 @@ public:
 	 */
 	virtual void UnregisterCustomClassLayout( FName ClassName );
 
-	UE_DEPRECATED(4.18, "This version of RegisterCustomPropertyTypeLayout has been deprecated.  For per-details instance customization call IDetailsView::RegisterInstancedCustomPropertyTypeLayout")
-	virtual void RegisterCustomPropertyTypeLayout(FName PropertyTypeName, FOnGetPropertyTypeCustomizationInstance PropertyTypeLayoutDelegate, TSharedPtr<IPropertyTypeIdentifier> Identifier, TSharedPtr<IDetailsView> ForSpecificInstance);
-
-	UE_DEPRECATED(4.18, "This version of UnregisterCustomPropertyTypeLayout has been deprecated.  For per-details instance customization call IDetailsView::UnregisterInstancedCustomPropertyTypeLayout")
-	virtual void UnregisterCustomPropertyTypeLayout(FName PropertyTypeName, TSharedPtr<IPropertyTypeIdentifier> InIdentifier, TSharedPtr<IDetailsView> ForSpecificInstance);
-
 	/**
 	 * Registers a property type customization
 	 * A property type is a specific FProperty type, a struct, or enum type
diff --git a/Engine/Source/Editor/SkeletonEditor/Private/SSkeletonTree.cpp b/Engine/Source/Editor/SkeletonEditor/Private/SSkeletonTree.cpp
--- a/Engine/Source/Editor/SkeletonEditor/Private/SSkeletonTree.cpp
+++ b/Engine/Source/Editor/SkeletonEditor/Private/SSkeletonTree.cpp
@@ -58,7 +58,6 @@
 #include "SkeletonTreeVirtualBoneItem.h"
 
 #include "BoneSelectionWidget.h"
-#include "BoneProxy.h"
 #include "SkeletonTreeSelection.h"
 #include "Widgets/Layout/SGridPanel.h"
 
@@ -194,13 +193,6 @@ void SSkeletonTree::Construct(const FArguments& InArgs, const TSharedRef<FEditab
 		RegisterOnSelectionChanged(FOnSkeletonTreeSelectionChanged::CreateRaw(PreviewScene.Pin().Get(), &IPersonaPreviewScene::HandleSkeletonTreeSelectionChanged));
 	}
 
-	PRAGMA_DISABLE_DEPRECATION_WARNINGS
-	if(InSkeletonTreeArgs.OnObjectSelected.IsBound())
-	{
-		RegisterOnObjectSelected(InSkeletonTreeArgs.OnObjectSelected);
-	}
-	PRAGMA_ENABLE_DEPRECATION_WARNINGS
-
 	if (InSkeletonTreeArgs.OnSelectionChanged.IsBound())
 	{
 		RegisterOnSelectionChanged(InSkeletonTreeArgs.OnSelectionChanged);
@@ -208,10 +200,6 @@ void SSkeletonTree::Construct(const FArguments& InArgs, const TSharedRef<FEditab
 
 	FCoreUObjectDelegates::OnPackageReloaded.AddSP(this, &SSkeletonTree::HandlePackageReloaded);
 
-	BoneProxy = NewObject<UBoneProxy>(GetTransientPackage());
-	BoneProxy->SkelMeshComponent = PreviewScene.IsValid() ? PreviewScene.Pin()->GetPreviewMeshComponent() : nullptr;
-	BoneProxy->bIsTickable = true;
-
 	// Create our pinned commands before we bind commands
 	IPinnedCommandListModule& PinnedCommandListModule = FModuleManager::LoadModuleChecked<IPinnedCommandListModule>(TEXT("PinnedCommandList"));
 	PinnedCommands = PinnedCommandListModule.CreatePinnedCommandList(ContextName);
@@ -1449,11 +1437,6 @@ void SSkeletonTree::OnSelectionChanged(TSharedPtr<ISkeletonTreeItem> Selection,
 							if (BoneIndex != INDEX_NONE)
 							{
 								GetPreviewScene()->SetSelectedBone(BoneName, SelectInfo);
-								BoneProxy->BoneName = BoneName;
-
-								PRAGMA_DISABLE_DEPRECATION_WARNINGS
-								OnObjectSelectedMulticast.Broadcast(BoneProxy);
-								PRAGMA_ENABLE_DEPRECATION_WARNINGS
 								break;
 							}
 						}
@@ -1465,18 +1448,10 @@ void SSkeletonTree::OnSelectionChanged(TSharedPtr<ISkeletonTreeItem> Selection,
 						USkeletalMeshSocket* Socket = SocketItem->GetSocket();
 						FSelectedSocketInfo SocketInfo(Socket, SocketItem->GetParentType() == ESocketParentType::Skeleton);
 						GetPreviewScene()->SetSelectedSocket(SocketInfo);
-
-						PRAGMA_DISABLE_DEPRECATION_WARNINGS
-						OnObjectSelectedMulticast.Broadcast(SocketInfo.Socket);
-						PRAGMA_ENABLE_DEPRECATION_WARNINGS
 					}
 					else if (Item->IsOfType<FSkeletonTreeAttachedAssetItem>())
 					{
 						GetPreviewScene()->DeselectAll();
-
-						PRAGMA_DISABLE_DEPRECATION_WARNINGS
-						OnObjectSelectedMulticast.Broadcast(nullptr);
-						PRAGMA_ENABLE_DEPRECATION_WARNINGS
 					}
 				}
 				PreviewComponent->PostInitMeshObject(PreviewComponent->MeshObject);
@@ -1490,10 +1465,6 @@ void SSkeletonTree::OnSelectionChanged(TSharedPtr<ISkeletonTreeItem> Selection,
 			// Tell the preview scene if the user ctrl-clicked the selected bone/socket to de-select it
 			GetPreviewScene()->DeselectAll();
 		}
-
-		PRAGMA_DISABLE_DEPRECATION_WARNINGS
-		OnObjectSelectedMulticast.Broadcast(nullptr);
-		PRAGMA_ENABLE_DEPRECATION_WARNINGS
 	}
 
 	TArrayView<TSharedPtr<ISkeletonTreeItem>> ArrayView(SelectedItems);
@@ -1722,10 +1693,6 @@ void SSkeletonTree::SetSelectedSocket( const FSelectedSocketInfo& SocketInfo )
 				SkeletonTreeView->RequestScrollIntoView(SkeletonRow);
 			}
 		}
-
-		PRAGMA_DISABLE_DEPRECATION_WARNINGS
-		OnObjectSelectedMulticast.Broadcast(SocketInfo.Socket);
-		PRAGMA_ENABLE_DEPRECATION_WARNINGS
 	}
 }
 
@@ -1747,11 +1714,6 @@ void SSkeletonTree::SetSelectedBone( const FName& BoneName, ESelectInfo::Type In
 				SkeletonTreeView->RequestScrollIntoView(SkeletonRow);
 			}
 		}
-
-		BoneProxy->BoneName = BoneName;
-		PRAGMA_DISABLE_DEPRECATION_WARNINGS
-		OnObjectSelectedMulticast.Broadcast(BoneProxy);
-		PRAGMA_ENABLE_DEPRECATION_WARNINGS
 	}
 }
 
@@ -1761,10 +1723,6 @@ void SSkeletonTree::DeselectAll()
 	{
 		TGuardValue<bool> RecursionGuard(bSelecting, true);
 		SkeletonTreeView->ClearSelection();
-
-		PRAGMA_DISABLE_DEPRECATION_WARNINGS
-		OnObjectSelectedMulticast.Broadcast(nullptr);
-		PRAGMA_ENABLE_DEPRECATION_WARNINGS
 	}
 }
 
@@ -2375,11 +2333,6 @@ ESkeletonTreeFilterResult SSkeletonTree::HandleFilterSkeletonTreeItem(const FSke
 	return Result;
 }
 
-void SSkeletonTree::AddReferencedObjects( FReferenceCollector& Collector )
-{
-	Collector.AddReferencedObject(BoneProxy);
-}
-
 void SSkeletonTree::HandleSelectedBoneChanged(const FName& InBoneName, ESelectInfo::Type InSelectInfo)
 {
 	SetSelectedBone(InBoneName, InSelectInfo);
diff --git a/Engine/Source/Editor/SkeletonEditor/Private/SSkeletonTree.h b/Engine/Source/Editor/SkeletonEditor/Private/SSkeletonTree.h
--- a/Engine/Source/Editor/SkeletonEditor/Private/SSkeletonTree.h
+++ b/Engine/Source/Editor/SkeletonEditor/Private/SSkeletonTree.h
@@ -47,7 +47,7 @@ enum class EBlendProfileMode : uint8;
 // SSkeletonTree
 
 PRAGMA_DISABLE_DEPRECATION_WARNINGS
-class SSkeletonTree : public ISkeletonTree, public FSelfRegisteringEditorUndoClient, public FGCObject
+class SSkeletonTree : public ISkeletonTree, public FSelfRegisteringEditorUndoClient
 {
 public:
 	SLATE_BEGIN_ARGS( SSkeletonTree )
@@ -81,16 +81,6 @@ public:
 	virtual void SelectItemsBy(TFunctionRef<bool(const TSharedRef<ISkeletonTreeItem>&, bool&)> Predicate) const override;
 	virtual void DuplicateAndSelectSocket(const FSelectedSocketInfo& SocketInfoToDuplicate, const FName& NewParentBoneName = FName()) override;
 
-	virtual void RegisterOnObjectSelected(const FOnObjectSelected& Delegate) override
-	{
-		OnObjectSelectedMulticast.Add(Delegate);
-	}
-
-	virtual void UnregisterOnObjectSelected(SWidget* Widget) override
-	{
-		OnObjectSelectedMulticast.RemoveAll(Widget);
-	}
-
 	virtual FDelegateHandle RegisterOnSelectionChanged(const FOnSkeletonTreeSelectionChanged& Delegate) override
 	{
 		return OnSelectionChangedMulticast.Add(Delegate);
@@ -110,13 +100,6 @@ public:
 	virtual void PostUndo(bool bSuccess) override;
 	virtual void PostRedo(bool bSuccess) override;
 
-	/** FGCObject interface */
-	virtual void AddReferencedObjects( FReferenceCollector& Collector ) override;
-	virtual FString GetReferencerName() const override
-	{
-		return TEXT("SSkeletonTree");
-	}
-
 	/** Creates the tree control and then populates */
 	void CreateTreeColumns();
 
@@ -408,9 +391,6 @@ private:
 	/** Delegate for when an item is selected */
 	FOnSkeletonTreeSelectionChangedMulticast OnSelectionChangedMulticast;
 
-	UE_DEPRECATED(4.17, "Please use OnSelectionChangedMulticast")
-	FOnObjectSelectedMulticast OnObjectSelectedMulticast;
-
 	/** Selection recursion guard flag */
 	bool bSelecting;
 
@@ -429,9 +409,6 @@ private:
 	/** Compiled filter search terms. */
 	TSharedPtr<class FTextFilterExpressionEvaluator> TextFilterPtr;
 
-	/** Proxy object used to display and edit bone transforms in details panels. Note this is only kept for backwards compatibility (used with OnObjectSelectedMulticast) */
-	class UBoneProxy* BoneProxy;
-
 	/** Whether to allow operations that modify the mesh */
 	bool bAllowMeshOperations;
 
diff --git a/Engine/Source/Editor/SkeletonEditor/Public/ISkeletonTree.h b/Engine/Source/Editor/SkeletonEditor/Public/ISkeletonTree.h
--- a/Engine/Source/Editor/SkeletonEditor/Public/ISkeletonTree.h
+++ b/Engine/Source/Editor/SkeletonEditor/Public/ISkeletonTree.h
@@ -53,17 +53,9 @@ struct FSkeletonTreeArgs
 		, bHideBonesByDefault(false)
 	{}
 
-	PRAGMA_DISABLE_DEPRECATION_WARNINGS
-	~FSkeletonTreeArgs()
-	{}
-	PRAGMA_ENABLE_DEPRECATION_WARNINGS
-
 	/** Delegate called by the tree when a socket is selected */
 	FOnSkeletonTreeSelectionChanged OnSelectionChanged;
 
-	UE_DEPRECATED(4.17, "Please use OnSelectionChanged")
-	FOnObjectSelected OnObjectSelected;
-
 	/** Delegate that allows custom filtering text to be shown on the filter button */
 	FOnGetFilterText OnGetFilterText;
 
@@ -163,11 +155,4 @@ public:
 
 	/** Get the pinned commands widget, if any, for this tree */
 	virtual TSharedPtr<IPinnedCommandList> GetPinnedCommandList() const = 0;
-
-	UE_DEPRECATED(4.17, "Please use RegisterOnSelectionChanged")
-	virtual void RegisterOnObjectSelected(const FOnObjectSelected& Delegate) = 0;
-
-	UE_DEPRECATED(4.17, "Please use UnregisterOnSelectionChanged")
-	virtual void UnregisterOnObjectSelected(SWidget* Widget) = 0;
-
 };
diff --git a/Engine/Source/Editor/UnrealEd/Classes/Editor/EditorEngine.h b/Engine/Source/Editor/UnrealEd/Classes/Editor/EditorEngine.h
--- a/Engine/Source/Editor/UnrealEd/Classes/Editor/EditorEngine.h
+++ b/Engine/Source/Editor/UnrealEd/Classes/Editor/EditorEngine.h
@@ -376,10 +376,6 @@ public:
 	UPROPERTY()
 	uint32 bEnableLODLocking:1;
 
-	/** If true, actors can be grouped and grouping rules will be maintained. When deactivated, any currently existing groups will still be preserved.*/
-	UE_DEPRECATED(4.17, "bGroupingActive has been deprecated.  Use UActorGroupingUtils::IsGroupingActive instead")
-	uint32 bGroupingActive:1;
-
 	UPROPERTY(config)
 	FString HeightMapExportClassName;
 
@@ -1954,14 +1950,6 @@ public:
 	 */
 	void ExportMap(UWorld* InWorld, const TCHAR* InFilename, bool bExportSelectedActorsOnly);
 
-	/**
-	 * Moves selected actors to the current level.
-	 *
-	 * @param	InLevel		The destination level.
-	 */
-	UE_DEPRECATED(4.17, "MoveSelectedActorsToLevel has been deprecated.  Use UEditorLevelUtils::MoveSelectedActorsToLevel instead")
-	void MoveSelectedActorsToLevel( ULevel* InLevel );
-
 	/**
 	 *	Returns list of all foliage types used in the world
 	 * 
@@ -3170,10 +3158,6 @@ protected:
 	// Handle requests from slate application to open assets.
 	bool HandleOpenAsset(UObject* Asset);
 
-public:
-	UE_DEPRECATED(4.17, "IsUsingWorldAssets is now always true, remove any code that assumes it could be false")
-	static bool IsUsingWorldAssets() { return true; }
-
 private:
 	/** Handler for when any asset is loaded in the editor */
 	void OnAssetLoaded( UObject* Asset );
diff --git a/Engine/Source/Editor/UnrealEd/Classes/Editor/UnrealEdEngine.h b/Engine/Source/Editor/UnrealEd/Classes/Editor/UnrealEdEngine.h
--- a/Engine/Source/Editor/UnrealEd/Classes/Editor/UnrealEdEngine.h
+++ b/Engine/Source/Editor/UnrealEd/Classes/Editor/UnrealEdEngine.h
@@ -488,43 +488,6 @@ public:
 	 */
 	virtual bool edactDeleteSelected( UWorld* InWorld, bool bVerifyDeletionCanHappen=true, bool bWarnAboutReferences = true, bool bWarnAboutSoftReferences = true) override final;
 
-	/**
-	 * Creates a new group from the current selection removing any existing groups.
-	 */
-	UE_DEPRECATED(4.17, "edactRegroupFromSelected is deprecated, use UActorGroupingUtils::GroupSelected")
-	virtual void edactRegroupFromSelected();
-
-	/**
-	 * Disbands any groups in the current selection, does not attempt to maintain any hierarchy
-	 */
-	UE_DEPRECATED(4.17, "edactUngroupFromSelected is deprecated, use UActorGroupingUtils::UngroupSelected")
-	virtual void edactUngroupFromSelected();
-
-	/**
-	 * Locks any groups in the current selection
-	 */
-	UE_DEPRECATED(4.17, "edactLockSelectedGroups is deprecated, use UActorGroupingUtils::LockSelectedGroups")
-	virtual void edactLockSelectedGroups();
-
-	/**
-	 * Unlocks any groups in the current selection
-	 */
-	UE_DEPRECATED(4.17, "edactUnlockSelectedGroups is deprecated, use UActorGroupingUtils::UnlockSelectedGroups")
-	virtual void edactUnlockSelectedGroups();
-
-	/**
-	 * Activates "Add to Group" mode which allows the user to select a group to append current selection
-	 */
-	UE_DEPRECATED(4.17, "edactAddToGroup is deprecated, use UActorGroupingUtils::AddSelectedToGroup")
-	virtual void edactAddToGroup();
-
-	/**
-	 * Removes any groups or actors in the current selection from their immediate parent.
-	 * If all actors/subgroups are removed, the parent group will be destroyed.
-	 */
-	UE_DEPRECATED(4.17, "edactRemoveFromGroup is deprecated, use UActorGroupingUtils::RemoveSelectedFromGroup")
-	virtual void edactRemoveFromGroup();
-	
 	/**
 	 * Copy selected actors to the clipboard.  Does not copy PrefabInstance actors or parts of Prefabs.
 	 * @note Final - Override CopyComponents or CopyActors instead.
diff --git a/Engine/Source/Editor/UnrealEd/Private/EditorGroup.cpp b/Engine/Source/Editor/UnrealEd/Private/EditorGroup.cpp
deleted file mode 100644
diff --git a/Engine/Source/Editor/UnrealEd/Private/EditorLevelUtils.cpp b/Engine/Source/Editor/UnrealEd/Private/EditorLevelUtils.cpp
--- a/Engine/Source/Editor/UnrealEd/Private/EditorLevelUtils.cpp
+++ b/Engine/Source/Editor/UnrealEd/Private/EditorLevelUtils.cpp
@@ -988,12 +988,6 @@ void UEditorLevelUtils::PrivateDestroyLevel(ULevel* InLevel)
 	}
 }
 
-ULevel* UEditorLevelUtils::CreateNewLevel(UWorld* InWorld, bool bMoveSelectedActorsIntoNewLevel, TSubclassOf<ULevelStreaming> LevelStreamingClass, const FString& DefaultFilename)
-{
-	ULevelStreaming* StreamingLevel = CreateNewStreamingLevelForWorld(*InWorld, LevelStreamingClass, DefaultFilename, bMoveSelectedActorsIntoNewLevel);
-	return StreamingLevel->GetLoadedLevel();
-}
-
 void UEditorLevelUtils::DeselectAllSurfacesInLevel(ULevel* InLevel)
 {
 	if (InLevel)
diff --git a/Engine/Source/Editor/UnrealEd/Private/EditorServer.cpp b/Engine/Source/Editor/UnrealEd/Private/EditorServer.cpp
--- a/Engine/Source/Editor/UnrealEd/Private/EditorServer.cpp
+++ b/Engine/Source/Editor/UnrealEd/Private/EditorServer.cpp
@@ -3271,14 +3271,6 @@ public:
 };
 } // namespace MoveSelectedActors
 
-
-void UEditorEngine::MoveSelectedActorsToLevel( ULevel* InDestLevel )
-{
-	// do the actual work...
-	UEditorLevelUtils::MoveSelectedActorsToLevel( InDestLevel );
-}
-
-
 TArray<UFoliageType*> UEditorEngine::GetFoliageTypesInWorld(UWorld* InWorld)
 {
 	TSet<UFoliageType*> FoliageSet;
diff --git a/Engine/Source/Editor/UnrealEd/Private/FileHelpers.cpp b/Engine/Source/Editor/UnrealEd/Private/FileHelpers.cpp
--- a/Engine/Source/Editor/UnrealEd/Private/FileHelpers.cpp
+++ b/Engine/Source/Editor/UnrealEd/Private/FileHelpers.cpp
@@ -4283,14 +4283,6 @@ bool FEditorFileUtils::SaveWorlds(UWorld* InWorld, const FString& RootPath, cons
 	return bSavedAll;
 }
 
-/**
- * DEPRECATED in version 4.18, Call FFileHelper::IsFilenameValidForSaving instead
- */
-bool FEditorFileUtils::IsFilenameValidForSaving( const FString& Filename, FText& OutError )
-{
-	return FFileHelper::IsFilenameValidForSaving(Filename, OutError);
-}
-
 void FEditorFileUtils::LoadDefaultMapAtStartup()
 {
 	FString EditorStartupMap;
diff --git a/Engine/Source/Editor/UnrealEd/Private/GeomFitUtils.cpp b/Engine/Source/Editor/UnrealEd/Private/GeomFitUtils.cpp
--- a/Engine/Source/Editor/UnrealEd/Private/GeomFitUtils.cpp
+++ b/Engine/Source/Editor/UnrealEd/Private/GeomFitUtils.cpp
@@ -642,12 +642,3 @@ void RefreshCollisionChangeComponentsOnly(UStaticMesh& StaticMesh)
 
 	FEditorSupportDelegates::RedrawAllViewports.Broadcast();
 }
-
-/* *************************** DEPRECATED ******************************** */
-void RefreshCollisionChange(const UStaticMesh* StaticMesh)
-{
-	if (StaticMesh)
-	{
-		RefreshCollisionChange(const_cast<UStaticMesh&>(*StaticMesh));
-	}
-}
diff --git a/Engine/Source/Editor/UnrealEd/Private/GeomFitUtils.h b/Engine/Source/Editor/UnrealEd/Private/GeomFitUtils.h
--- a/Engine/Source/Editor/UnrealEd/Private/GeomFitUtils.h
+++ b/Engine/Source/Editor/UnrealEd/Private/GeomFitUtils.h
@@ -144,7 +144,3 @@ UNREALED_API void RefreshCollisionChanges(const TArray<UStaticMesh*>& StaticMesh
  * @param	StaticMesh	StaticMesh that collision has been changed for
  */
 UNREALED_API void RefreshCollisionChangeComponentsOnly(UStaticMesh& StaticMesh);
-
-UE_DEPRECATED(4.15, "This version of RefreshCollisionChange is deprecated. Please use RefreshCollisionChange(UStaticMesh&) instead.")
-UNREALED_API void RefreshCollisionChange(const UStaticMesh* StaticMesh);
-	
diff --git a/Engine/Source/Editor/UnrealEd/Private/ObjectTools.cpp b/Engine/Source/Editor/UnrealEd/Private/ObjectTools.cpp
--- a/Engine/Source/Editor/UnrealEd/Private/ObjectTools.cpp
+++ b/Engine/Source/Editor/UnrealEd/Private/ObjectTools.cpp
@@ -4200,28 +4200,6 @@ namespace ObjectTools
 		}
 	}
 
-	/**
-	 * Exports the specified objects to file.
-	 *
-	 * @param	ObjectsToExport					The set of objects to export.
-	 * @param	bPromptIndividualFilenames		If true, prompt individually for filenames.  If false, bulk export to a single directory.
-	 * @param	ExportPath						receives the value of the path the user chose for exporting.
-	 * @param	bUseProvidedExportPath			If true and out_ExportPath is specified, use the value in out_ExportPath as the export path w/o prompting for a directory when applicable
-	 */
-	void ExportObjects(const TArray<UObject*>& ObjectsToExport, bool bPromptIndividualFilenames, FString* ExportPath/*=NULL*/, bool bUseProvidedExportPath /*= false*/ )
-	{
-		FAssetToolsModule& AssetToolsModule = FModuleManager::LoadModuleChecked<FAssetToolsModule>(TEXT("AssetTools"));
-
-		if (ExportPath && bUseProvidedExportPath && !bPromptIndividualFilenames)
-		{
-			AssetToolsModule.Get().ExportAssets(ObjectsToExport, *ExportPath);
-		}
-		else
-		{
-			AssetToolsModule.Get().ExportAssetsWithDialog(ObjectsToExport, bPromptIndividualFilenames);
-		}
-	}
-
 	/**
 	 * Tags objects which are in use by levels specified by the search option
 	 *
diff --git a/Engine/Source/Editor/UnrealEd/Public/EditorLevelUtils.h b/Engine/Source/Editor/UnrealEd/Public/EditorLevelUtils.h
--- a/Engine/Source/Editor/UnrealEd/Public/EditorLevelUtils.h
+++ b/Engine/Source/Editor/UnrealEd/Public/EditorLevelUtils.h
@@ -324,20 +324,6 @@ public:
 	UFUNCTION(BlueprintCallable, Category = "Editor Scripting | Level Utility")
 	static const TArray<ULevel*> GetLevels(UWorld* World);
 
-	UE_DEPRECATED(4.17, "The CreateNewLevel method to create streaming levels has been deprecated. Use CreateNewStreamingLevelForWorld instead")
-	static UNREALED_API ULevel* CreateNewLevel(UWorld* InWorld, bool bMoveSelectedActorsIntoNewLevel, TSubclassOf<ULevelStreaming> LevelStreamingClass, const FString& DefaultFilename = TEXT(""));
-
-
-	/**
-	 * Moves the specified list of actors to the specified level
-	 *
-	 * @param	ActorsToMove		List of actors to move
-	 * @param	DestLevelStreaming	The level streaming object associated with the destination level
-	 * @param	OutNumMovedActors	The number of actors that were successfully moved to the new level
-	 */
-	UE_DEPRECATED(4.17, "The MovesActorsToLevel method has been deprecated. Use MoveActorsToLevel instead")
-	static void MovesActorsToLevel(TArray< AActor* >& ActorsToMove, ULevelStreaming* DestLevelStreaming, int32& OutNumMovedActors);
-
 private:
 	/**
 	* Utility methods used by RemoveLevelsFromWorld
diff --git a/Engine/Source/Editor/UnrealEd/Public/FileHelpers.h b/Engine/Source/Editor/UnrealEd/Public/FileHelpers.h
--- a/Engine/Source/Editor/UnrealEd/Public/FileHelpers.h
+++ b/Engine/Source/Editor/UnrealEd/Public/FileHelpers.h
@@ -524,16 +524,6 @@ public:
 	 */
 	UNREALED_API static bool PromptToCheckoutLevels(bool bCheckDirty, ULevel* SpecificLevelToCheckOut);
 
-	/**
-	 * Checks to see if a filename is valid for saving.
-	 * A filename must be under FPlatformMisc::GetMaxPathLength() to be saved
-	 *
-	 * @param Filename	Filename, with or without path information, to check.
-	 * @param OutError	If an error occurs, this is the reason why
-	 */
-	UE_DEPRECATED(4.18, "Call FFileHelper::IsFilenameValidForSaving instead")
-	UNREALED_API static bool IsFilenameValidForSaving(const FString& Filename, FText& OutError);
-
 	/** Loads a simple example map */
 	UNREALED_API static void LoadDefaultMapAtStartup();
 	
diff --git a/Engine/Source/Editor/UnrealEd/Public/ObjectTools.h b/Engine/Source/Editor/UnrealEd/Public/ObjectTools.h
--- a/Engine/Source/Editor/UnrealEd/Public/ObjectTools.h
+++ b/Engine/Source/Editor/UnrealEd/Public/ObjectTools.h
@@ -484,19 +484,6 @@ namespace ObjectTools
 	 */
 	UNREALED_API void GetDirectoryFromObjectPath(const UObject* Obj, FString& OutResult);
 
-
-	/**
-	 * Exports the specified objects to file.
-	 *
-	 * @param	ObjectsToExport					The set of objects to export.
-	 * @param	bPromptIndividualFilenames		If true, prompt individually for filenames.  If false, bulk export to a single directory.
-	 * @param	ExportPath						receives the value of the path the user chose for exporting.
-	 * @param	bUseProvidedExportPath			If true and out_ExportPath is specified, use the value in out_ExportPath as the export path w/o prompting for a directory when applicable
-	 */
-	UE_DEPRECATED(4.17, "ObjectTools::ExportObjects is deprecated.  Use AssetTools::ExportObjects instead")
-	UNREALED_API void ExportObjects( const TArray<UObject*>& ObjectsToExport, bool bPromptIndividualFilenames, FString* ExportPath = NULL, bool bUseProvidedExportPath = false );
-
-
 	/** Options for in use object tagging */
 	enum EInUseSearchOption
 	{
diff --git a/Engine/Source/Runtime/AIModule/Classes/AITypes.h b/Engine/Source/Runtime/AIModule/Classes/AITypes.h
--- a/Engine/Source/Runtime/AIModule/Classes/AITypes.h
+++ b/Engine/Source/Runtime/AIModule/Classes/AITypes.h
@@ -533,12 +533,6 @@ struct AIMODULE_API FAIMoveRequest
 	bool UpdateGoalLocation(const FVector& NewLocation) const;
 	FString ToString() const;
 
-	UE_DEPRECATED(4.13, "This function is deprecated, please use SetReachTestIncludesAgentRadius instead.")
-	FAIMoveRequest& SetStopOnOverlap(bool bStop);
-	
-	UE_DEPRECATED(4.13, "This function is deprecated, please use IsReachTestIncludingAgentRadius instead.")
-	bool CanStopOnOverlap() const;
-
 protected:
 
 	/** move goal: actor */
diff --git a/Engine/Source/Runtime/AIModule/Classes/BehaviorTree/BTDecorator.h b/Engine/Source/Runtime/AIModule/Classes/BehaviorTree/BTDecorator.h
--- a/Engine/Source/Runtime/AIModule/Classes/BehaviorTree/BTDecorator.h
+++ b/Engine/Source/Runtime/AIModule/Classes/BehaviorTree/BTDecorator.h
@@ -138,12 +138,6 @@ protected:
 		bNotifyDeactivation = !TIsSame<decltype(&UBTDecorator::OnNodeDeactivation), OnNodeDeactivation>::Value;
 		bNotifyProcessed = !TIsSame<decltype(&UBTDecorator::OnNodeProcessed), OnNodeProcessed>::Value;
 	}
-
-	//----------------------------------------------------------------------//
-	// DEPRECATED
-	//----------------------------------------------------------------------//
-	UE_DEPRECATED(4.12, "This function is deprecated, please use InitializeParentLink instead.")
-	void InitializeDecorator(uint8 InChildIndex);
 };
 
 #define INIT_DECORATOR_NODE_NOTIFY_FLAGS() \
diff --git a/Engine/Source/Runtime/AIModule/Classes/BehaviorTree/BTFunctionLibrary.h b/Engine/Source/Runtime/AIModule/Classes/BehaviorTree/BTFunctionLibrary.h
--- a/Engine/Source/Runtime/AIModule/Classes/BehaviorTree/BTFunctionLibrary.h
+++ b/Engine/Source/Runtime/AIModule/Classes/BehaviorTree/BTFunctionLibrary.h
@@ -23,9 +23,6 @@ namespace FBTNodeBPImplementationHelper
 	/** checks if given object implements GenericEventName and/or AIEventName BP events, and returns an result as flags set on return integer 
 	 *	@return flags set in returned integer indicate kinds of events implemented by given object */
 	AIMODULE_API int32 CheckEventImplementationVersion(FName GenericEventName, FName AIEventName, const UObject& Object, const UClass& StopAtClass);
-
-	UE_DEPRECATED(4.11, "This version of CheckEventImplementationVersion is deprecated. Please use the one taking reference to UObject and StopAtClass rather than a pointers.")
-	AIMODULE_API int32 CheckEventImplementationVersion(FName GenericEventName, FName AIEventName, const UObject* Ob, const UClass* StopAtClass);
 }
 
 UCLASS(meta=(RestrictedToClasses="BTNode"))
diff --git a/Engine/Source/Runtime/AIModule/Classes/BehaviorTree/BlackboardComponent.h b/Engine/Source/Runtime/AIModule/Classes/BehaviorTree/BlackboardComponent.h
--- a/Engine/Source/Runtime/AIModule/Classes/BehaviorTree/BlackboardComponent.h
+++ b/Engine/Source/Runtime/AIModule/Classes/BehaviorTree/BlackboardComponent.h
@@ -84,14 +84,6 @@ public:
 	/** resume observer change notifications and, optionally, process the queued observation list */
 	void ResumeObserverNotifications(bool bSendQueuedObserverNotifications);
 
-	/** pause change notifies and add them to queue */
-	UE_DEPRECATED(4.15, "Please call PauseObserverUpdates.")
-	void PauseUpdates();
-
-	/** resume change notifies and process queued list */
-	UE_DEPRECATED(4.15, "Please call ResumeObserverNotifications.")
-	void ResumeUpdates();
-
 	/** @return associated behavior tree component */
 	UBrainComponent* GetBrainComponent() const;
 
diff --git a/Engine/Source/Runtime/AIModule/Classes/BehaviorTree/BlackboardData.h b/Engine/Source/Runtime/AIModule/Classes/BehaviorTree/BlackboardData.h
--- a/Engine/Source/Runtime/AIModule/Classes/BehaviorTree/BlackboardData.h
+++ b/Engine/Source/Runtime/AIModule/Classes/BehaviorTree/BlackboardData.h
@@ -177,8 +177,4 @@ protected:
 
 	/** @return key ID from name */
 	FBlackboard::FKey InternalGetKeyID(const FName& KeyName, EKeyLookupMode LookupMode) const;
-
-	/** check if parent chain contains given blackboard data */
-	UE_DEPRECATED(4.14, "This function is deprecated, please use IsChildOf instead.")
-	bool HasParent(const UBlackboardData* TestParent) const;
 };
diff --git a/Engine/Source/Runtime/AIModule/Classes/Navigation/CrowdFollowingComponent.h b/Engine/Source/Runtime/AIModule/Classes/Navigation/CrowdFollowingComponent.h
--- a/Engine/Source/Runtime/AIModule/Classes/Navigation/CrowdFollowingComponent.h
+++ b/Engine/Source/Runtime/AIModule/Classes/Navigation/CrowdFollowingComponent.h
@@ -141,9 +141,6 @@ class AIMODULE_API UCrowdFollowingComponent : public UPathFollowingComponent, pu
 	virtual void DescribeSelfToVisLog(struct FVisualLogEntry* Snapshot) const override;
 #endif // ENABLE_VISUAL_LOG
 
-	UE_DEPRECATED(4.11, "Use SetCrowdSimulationState function instead.")
-	virtual void SetCrowdSimulation(bool bEnable);
-
 	UE_DEPRECATED_FORGAME(4.16, "Use ApplyCrowdAgentVelocity function with bIsNearEndOfPath param instead.")
 	virtual void ApplyCrowdAgentVelocity(const FVector& NewVelocity, const FVector& DestPathCorner, bool bTraversingLink) {}
 
@@ -185,9 +182,6 @@ protected:
 	/** if set, move velocity will be updated in every tick */
 	uint8 bUpdateDirectMoveVelocity : 1;
 
-	UE_DEPRECATED(4.11, "Please use IsCrowdSimulationEnabled(), SetCrowdSimulationState() and SimulationState member for initialization.")
-	uint8 bEnableCrowdSimulation : 1;
-
 	/** set when agent is registered in crowd simulation (either controlled or an obstacle) */
 	uint8 bRegisteredWithCrowdSimulation : 1;
 
@@ -248,8 +242,6 @@ protected:
 	bool HasMovedDuringPause() const;
 	void UpdateCachedDirections(const FVector& NewVelocity, const FVector& NextPathCorner, bool bTraversingLink);
 
-	UE_DEPRECATED(4.12, "This function is now deprecated and was renamed to ShouldTrackMovingGoal.")
-	virtual bool UpdateCachedGoal(FVector& NewGoalPos);
 	virtual bool ShouldTrackMovingGoal(FVector& OutGoalLocation) const;
 	
 	void RegisterCrowdAgent();
diff --git a/Engine/Source/Runtime/AIModule/Classes/Navigation/MetaNavMeshPath.h b/Engine/Source/Runtime/AIModule/Classes/Navigation/MetaNavMeshPath.h
--- a/Engine/Source/Runtime/AIModule/Classes/Navigation/MetaNavMeshPath.h
+++ b/Engine/Source/Runtime/AIModule/Classes/Navigation/MetaNavMeshPath.h
@@ -129,8 +129,4 @@ protected:
 
 	/** switch to next waypoint  */
 	bool MoveToNextSection(const FVector& AgentLocation);
-
-public:
-	UE_DEPRECATED(4.16, "Waypoint type of MetaPaths has been changed. Use the new GetWaypointArray")
-	TArray<FVector> GetWaypoints() const;
 };
diff --git a/Engine/Source/Runtime/AIModule/Classes/Navigation/PathFollowingComponent.h b/Engine/Source/Runtime/AIModule/Classes/Navigation/PathFollowingComponent.h
--- a/Engine/Source/Runtime/AIModule/Classes/Navigation/PathFollowingComponent.h
+++ b/Engine/Source/Runtime/AIModule/Classes/Navigation/PathFollowingComponent.h
@@ -410,39 +410,12 @@ class AIMODULE_API UPathFollowingComponent : public UActorComponent, public IAIR
 	static void LogPathHelper(const AActor* LogOwner, FNavPathSharedPtr InLogPath, const AActor* LogGoalActor);
 	static void LogPathHelper(const AActor* LogOwner, FNavigationPath* InLogPath, const AActor* LogGoalActor);
 
-	UE_DEPRECATED(4.12, "This function is now deprecated and replaced with HandlePathUpdateEvent. Receiving new path pointer for the same move request is no longer supported, please either update data within current path and call FNavigationPath::DoneUpdating or start new move request.")
-	virtual bool UpdateMove(FNavPathSharedPtr Path, FAIRequestID RequestID = FAIRequestID::CurrentRequest);
-
-	UE_DEPRECATED(4.13, "This function is now deprecated, please use version with FAIMoveRequest parameter instead. Any observers needs to register with OnRequestFinished mutlicast delegate now.")
-	virtual FAIRequestID RequestMove(FNavPathSharedPtr Path, FRequestCompletedSignature OnComplete, const AActor* DestinationActor = NULL, float AcceptanceRadius = UPathFollowingComponent::DefaultAcceptanceRadius, bool bStopOnOverlap = true, FCustomMoveSharedPtr GameData = NULL);
-
-	UE_DEPRECATED(4.13, "This function is now deprecated, please use version with FAIMoveRequest parameter instead.")
-	FAIRequestID RequestMove(FNavPathSharedPtr InPath, const AActor* InDestinationActor = NULL, float InAcceptanceRadius = UPathFollowingComponent::DefaultAcceptanceRadius, bool InStopOnOverlap = true, FCustomMoveSharedPtr InGameData = NULL);
-
-	UE_DEPRECATED(4.13, "This function is now deprecated, please use version with EPathFollowingResultDetails parameter instead.")
-	virtual void AbortMove(const FString& Reason, FAIRequestID RequestID = FAIRequestID::CurrentRequest, bool bResetVelocity = true, bool bSilent = false, uint8 MessageFlags = 0);
-
-	UE_DEPRECATED(4.13, "This function is now deprecated, please use version with EPathFollowingVelocityMode parameter instead.")
-	virtual void PauseMove(FAIRequestID RequestID = FAIRequestID::CurrentRequest, bool bResetVelocity = true);
-
 	UFUNCTION(BlueprintCallable, Category="AI|Components|PathFollowing", meta = (DeprecatedFunction, DeprecationMessage = "This function is now deprecated, please use AIController.GetMoveStatus instead"))
 	EPathFollowingAction::Type GetPathActionType() const;
 
 	UFUNCTION(BlueprintCallable, Category="AI|Components|PathFollowing", meta = (DeprecatedFunction, DeprecationMessage = "This function is now deprecated, please use AIController.GetImmediateMoveDestination instead"))
 	FVector GetPathDestination() const;
 
-	UE_DEPRECATED(4.13, "This function is now deprecated, please version with FPathFollowingResult parameter instead.")
-	virtual void OnPathFinished(EPathFollowingResult::Type Result);
-
-	UE_DEPRECATED(4.13, "This function is now deprecated and no longer supported.")
-	int32 OptimizeSegmentVisibility(int32 StartIndex) { return StartIndex + 1; }
-
-	UE_DEPRECATED(4.13, "This function is now deprecated, please use version with EPathFollowingReachMode parameter instead.")
-	bool HasReached(const FVector& TestPoint, float AcceptanceRadius = UPathFollowingComponent::DefaultAcceptanceRadius, bool bExactSpot = false) const;
-
-	UE_DEPRECATED(4.13, "This function is now deprecated, please use version with EPathFollowingReachMode parameter instead.")
-	bool HasReached(const AActor& TestGoal, float AcceptanceRadius = UPathFollowingComponent::DefaultAcceptanceRadius, bool bExactSpot = false, bool bUseNavAgentGoalLocation = true) const;
-
 #if WITH_EDITORONLY_DATA
 	// This delegate is now deprecated, please use OnRequestFinished instead
 	FMoveCompletedSignature OnMoveFinished_DEPRECATED;
diff --git a/Engine/Source/Runtime/AIModule/Classes/Perception/AIPerceptionComponent.h b/Engine/Source/Runtime/AIModule/Classes/Perception/AIPerceptionComponent.h
--- a/Engine/Source/Runtime/AIModule/Classes/Perception/AIPerceptionComponent.h
+++ b/Engine/Source/Runtime/AIModule/Classes/Perception/AIPerceptionComponent.h
@@ -408,9 +408,6 @@ public:
 	FActorPerceptionInfoUpdatedDelegate OnTargetPerceptionInfoUpdated;
 
 protected:
-	UE_DEPRECATED(4.11, "Function has been renamed and made public. Please use UpdatePerceptionAllowList instead")
-	void UpdatePerceptionFilter(FAISenseID Channel, bool bNewValue);
-
 	FActorPerceptionContainer& GetPerceptualData() { return PerceptualData; }
 	const FActorPerceptionContainer& GetPerceptualData() const { return PerceptualData; }
 
diff --git a/Engine/Source/Runtime/AIModule/Classes/Perception/AISense.h b/Engine/Source/Runtime/AIModule/Classes/Perception/AISense.h
--- a/Engine/Source/Runtime/AIModule/Classes/Perception/AISense.h
+++ b/Engine/Source/Runtime/AIModule/Classes/Perception/AISense.h
@@ -154,11 +154,4 @@ protected:
 
 	/** To be called only for BP-generated classes */
 	void ForceSenseID(FAISenseID SenseID);
-
-	UE_DEPRECATED(4.13, "Function removed, use AISenseConfig.GetDebugColor() instead.")
-	FColor GetDebugColor() const { return FColor::White; }
-	UE_DEPRECATED(4.13, "Function removed, use AISenseConfig.GetSenseName() instead.")
-	FString GetDebugName() const { return FString(); }
-	UE_DEPRECATED(4.13, "Function removed, use AISenseConfig.DescribeSelfToGameplayDebugger() instead.")
-	virtual FString GetDebugLegend() const { return FString(); };
 };
diff --git a/Engine/Source/Runtime/AIModule/Classes/Tasks/AITask_MoveTo.h b/Engine/Source/Runtime/AIModule/Classes/Tasks/AITask_MoveTo.h
--- a/Engine/Source/Runtime/AIModule/Classes/Tasks/AITask_MoveTo.h
+++ b/Engine/Source/Runtime/AIModule/Classes/Tasks/AITask_MoveTo.h
@@ -36,9 +36,6 @@ public:
 		float AcceptanceRadius = -1.f, EAIOptionFlag::Type StopOnOverlap = EAIOptionFlag::Default, EAIOptionFlag::Type AcceptPartialPath = EAIOptionFlag::Default,
 		bool bUsePathfinding = true, bool bLockAILogic = true, bool bUseContinuosGoalTracking = false, EAIOptionFlag::Type ProjectGoalOnNavigation = EAIOptionFlag::Default);
 
-	UE_DEPRECATED(4.12, "This function is now depreacted, please use version with FAIMoveRequest parameter")
-	void SetUp(AAIController* Controller, FVector GoalLocation, AActor* GoalActor = nullptr, float AcceptanceRadius = -1.f, bool bUsePathfinding = true, EAIOptionFlag::Type StopOnOverlap = EAIOptionFlag::Default, EAIOptionFlag::Type AcceptPartialPath = EAIOptionFlag::Default);
-
 	/** Allows custom move request tweaking. Note that all MoveRequest need to 
 	 *	be performed before PerformMove is called. */
 	FAIMoveRequest& GetMoveRequestRef() { return MoveRequest; }
diff --git a/Engine/Source/Runtime/AIModule/Private/AITypes.cpp b/Engine/Source/Runtime/AIModule/Private/AITypes.cpp
--- a/Engine/Source/Runtime/AIModule/Private/AITypes.cpp
+++ b/Engine/Source/Runtime/AIModule/Private/AITypes.cpp
@@ -227,12 +227,3 @@ FString FAIMoveRequest::ToString() const
 		);
 }
 
-FAIMoveRequest& FAIMoveRequest::SetStopOnOverlap(bool bStop)
-{
-	return SetReachTestIncludesAgentRadius(bStop);
-}
-
-bool FAIMoveRequest::CanStopOnOverlap() const
-{
-	return IsReachTestIncludingAgentRadius();
-}
diff --git a/Engine/Source/Runtime/AIModule/Private/BehaviorTree/BTDecorator.cpp b/Engine/Source/Runtime/AIModule/Private/BehaviorTree/BTDecorator.cpp
--- a/Engine/Source/Runtime/AIModule/Private/BehaviorTree/BTDecorator.cpp
+++ b/Engine/Source/Runtime/AIModule/Private/BehaviorTree/BTDecorator.cpp
@@ -209,10 +209,3 @@ void UBTDecorator::UpdateFlowAbortMode()
 #endif
 }
 
-//----------------------------------------------------------------------//
-// DEPRECATED
-//----------------------------------------------------------------------//
-void UBTDecorator::InitializeDecorator(uint8 InChildIndex)
-{
-	InitializeParentLink(InChildIndex);
-}
diff --git a/Engine/Source/Runtime/AIModule/Private/BehaviorTree/BTFunctionLibrary.cpp b/Engine/Source/Runtime/AIModule/Private/BehaviorTree/BTFunctionLibrary.cpp
--- a/Engine/Source/Runtime/AIModule/Private/BehaviorTree/BTFunctionLibrary.cpp
+++ b/Engine/Source/Runtime/AIModule/Private/BehaviorTree/BTFunctionLibrary.cpp
@@ -30,11 +30,6 @@ namespace FBTNodeBPImplementationHelper
 
 		return (bGeneric ? Generic : NoImplementation) | (bAI ? AISpecific : NoImplementation);
 	}
-
-	int32 CheckEventImplementationVersion(FName GenericEventName, FName AIEventName, const UObject* Ob, const UClass* StopAtClass)
-	{
-		return (Ob && StopAtClass) ? CheckEventImplementationVersion(GenericEventName, AIEventName, *Ob, *StopAtClass) : NoImplementation;
-	}
 }
 
 //----------------------------------------------------------------------//
diff --git a/Engine/Source/Runtime/AIModule/Private/BehaviorTree/BlackboardComponent.cpp b/Engine/Source/Runtime/AIModule/Private/BehaviorTree/BlackboardComponent.cpp
--- a/Engine/Source/Runtime/AIModule/Private/BehaviorTree/BlackboardComponent.cpp
+++ b/Engine/Source/Runtime/AIModule/Private/BehaviorTree/BlackboardComponent.cpp
@@ -397,23 +397,6 @@ void UBlackboardComponent::ResumeObserverNotifications(bool bSendQueuedObserverN
 	QueuedUpdates.Empty();
 }
 
-void UBlackboardComponent::PauseUpdates()
-{
-	bPausedNotifies = true;
-}
-
-void UBlackboardComponent::ResumeUpdates()
-{
-	bPausedNotifies = false;
-
-	for (int32 UpdateIndex = 0; UpdateIndex < QueuedUpdates.Num(); UpdateIndex++)
-	{
-		NotifyObservers(QueuedUpdates[UpdateIndex]);
-	}
-
-	QueuedUpdates.Empty();
-}
-
 void UBlackboardComponent::NotifyObservers(FBlackboard::FKey KeyID) const
 {
 	TMultiMap<uint8, FOnBlackboardChangeNotificationInfo>::TKeyIterator KeyIt(Observers, KeyID);
diff --git a/Engine/Source/Runtime/AIModule/Private/BehaviorTree/BlackboardData.cpp b/Engine/Source/Runtime/AIModule/Private/BehaviorTree/BlackboardData.cpp
--- a/Engine/Source/Runtime/AIModule/Private/BehaviorTree/BlackboardData.cpp
+++ b/Engine/Source/Runtime/AIModule/Private/BehaviorTree/BlackboardData.cpp
@@ -307,15 +307,3 @@ bool UBlackboardData::IsChildOf(const UBlackboardData& OtherAsset) const
 	return (TmpParent == &OtherAsset);
 }
 
-//----------------------------------------------------------------------//
-// DEPRECATED
-//----------------------------------------------------------------------//
-bool UBlackboardData::HasParent(const UBlackboardData* TestParent) const
-{
-	if (Parent == TestParent)
-	{
-		return true;
-	}
-
-	return Parent ? Parent->HasParent(TestParent) : false;
-}
diff --git a/Engine/Source/Runtime/AIModule/Private/BehaviorTree/BlueprintNodeHelpers.cpp b/Engine/Source/Runtime/AIModule/Private/BehaviorTree/BlueprintNodeHelpers.cpp
--- a/Engine/Source/Runtime/AIModule/Private/BehaviorTree/BlueprintNodeHelpers.cpp
+++ b/Engine/Source/Runtime/AIModule/Private/BehaviorTree/BlueprintNodeHelpers.cpp
@@ -299,13 +299,5 @@ namespace BlueprintNodeHelpers
 			}
 		}
 	}
-
-	//----------------------------------------------------------------------//
-	// DEPRECATED
-	//----------------------------------------------------------------------//
-	bool HasBlueprintFunction(FName FuncName, const UObject* Object, const UClass* StopAtClass)
-	{
-		return Object && StopAtClass && HasBlueprintFunction(FuncName, *Object, *StopAtClass);
-	}
 }
 
diff --git a/Engine/Source/Runtime/AIModule/Private/Navigation/CrowdFollowingComponent.cpp b/Engine/Source/Runtime/AIModule/Private/Navigation/CrowdFollowingComponent.cpp
--- a/Engine/Source/Runtime/AIModule/Private/Navigation/CrowdFollowingComponent.cpp
+++ b/Engine/Source/Runtime/AIModule/Private/Navigation/CrowdFollowingComponent.cpp
@@ -375,11 +375,6 @@ void UCrowdFollowingComponent::UpdateDestinationForMovingGoal(const FVector& New
 	CurrentDestination.Set(Path->GetBaseActor(), NewDestination);
 }
 
-bool UCrowdFollowingComponent::UpdateCachedGoal(FVector& NewGoalPos)
-{
-	return ShouldTrackMovingGoal(NewGoalPos);
-}
-
 void UCrowdFollowingComponent::ApplyCrowdAgentVelocity(const FVector& NewVelocity, const FVector& DestPathCorner, bool bTraversingLink, bool bIsNearEndOfPath)
 {
 	bCanCheckMovingTooFar = !bTraversingLink && bIsNearEndOfPath;
@@ -416,11 +411,6 @@ void UCrowdFollowingComponent::ApplyCrowdAgentPosition(const FVector& NewPositio
 	// base implementation does nothing
 }
 
-void UCrowdFollowingComponent::SetCrowdSimulation(bool bEnable)
-{
-	SetCrowdSimulationState(bEnable ? ECrowdSimulationState::Enabled : ECrowdSimulationState::ObstacleOnly);
-}
-
 void UCrowdFollowingComponent::SetCrowdSimulationState(ECrowdSimulationState NewState)
 {
 	if (NewState == SimulationState)
diff --git a/Engine/Source/Runtime/AIModule/Private/Navigation/MetaNavMeshPath.cpp b/Engine/Source/Runtime/AIModule/Private/Navigation/MetaNavMeshPath.cpp
--- a/Engine/Source/Runtime/AIModule/Private/Navigation/MetaNavMeshPath.cpp
+++ b/Engine/Source/Runtime/AIModule/Private/Navigation/MetaNavMeshPath.cpp
@@ -274,16 +274,3 @@ void FMetaNavMeshPath::DebugDraw(const ANavigationData* NavData, const FColor Pa
 #endif // ENABLE_DRAW_DEBUG
 }
 
-//----------------------------------------------------------------------//
-// DEPRECATED
-//----------------------------------------------------------------------//
-TArray<FVector> FMetaNavMeshPath::GetWaypoints() const 
-{ 
-	TArray<FVector> VectorWaypoints;
-	VectorWaypoints.Reserve(Waypoints.Num());
-	for (const FMetaPathWayPoint& Waypoint : Waypoints)
-	{
-		VectorWaypoints.Add(Waypoint);
-	}
-	return VectorWaypoints; 
-}
diff --git a/Engine/Source/Runtime/AIModule/Private/Navigation/PathFollowingComponent.cpp b/Engine/Source/Runtime/AIModule/Private/Navigation/PathFollowingComponent.cpp
--- a/Engine/Source/Runtime/AIModule/Private/Navigation/PathFollowingComponent.cpp
+++ b/Engine/Source/Runtime/AIModule/Private/Navigation/PathFollowingComponent.cpp
@@ -430,20 +430,6 @@ FAIRequestID UPathFollowingComponent::RequestMove(const FAIMoveRequest& RequestD
 	return MoveId;
 }
 
-bool UPathFollowingComponent::UpdateMove(FNavPathSharedPtr InPath, FAIRequestID RequestID)
-{
-	if (InPath == Path && RequestID.IsEquivalent(CurrentRequestId))
-	{
-		return HandlePathUpdateEvent();
-	}
-
-	UE_VLOG(GetOwner(), LogPathFollowing, Error, TEXT("UpdateMove is ignored! path:%s requestId:%s"),
-		(InPath == Path) ? TEXT("ok") : TEXT("DIFFERENT"),
-		RequestID.IsEquivalent(CurrentRequestId) ? TEXT("ok") : TEXT("NOT CURRENT"));
-
-	return false;
-}
-
 void UPathFollowingComponent::AbortMove(const UObject& Instigator, FPathFollowingResultFlags::Type AbortFlags, FAIRequestID RequestID, EPathFollowingVelocityMode VelocityMode)
 {
 	UE_VLOG(GetOwner(), LogPathFollowing, Log, TEXT("AbortMove: RequestID(%u) Instigator(%s)"), RequestID, *Instigator.GetName());
@@ -1866,60 +1852,6 @@ void UPathFollowingComponent::OnWaitingPathTimeout()
 }
 
 // deprecated functions
-FAIRequestID UPathFollowingComponent::RequestMove(FNavPathSharedPtr InPath, FRequestCompletedSignature OnComplete, const AActor* InDestinationActor, float InAcceptanceRadius, bool bInStopOnOverlap, FCustomMoveSharedPtr InGameData)
-{
-	FAIMoveRequest MoveReq;
-	if (InDestinationActor)
-	{
-		MoveReq.SetGoalActor(InDestinationActor);
-	}
-	else
-	{
-		MoveReq.SetGoalLocation(InPath.IsValid() && InPath->GetPathPoints().Num() ? InPath->GetPathPoints().Last().Location : FAISystem::InvalidLocation);
-	}
-
-	MoveReq.SetAcceptanceRadius(InAcceptanceRadius);
-	MoveReq.SetReachTestIncludesAgentRadius(bInStopOnOverlap);
-	MoveReq.SetUserData(InGameData);
-
-	return RequestMove(MoveReq, InPath);
-}
-
-FAIRequestID UPathFollowingComponent::RequestMove(FNavPathSharedPtr InPath, const AActor* InDestinationActor, float InAcceptanceRadius, bool bInStopOnOverlap, FCustomMoveSharedPtr InGameData)
-{
-	FAIMoveRequest MoveReq;
-	if (InDestinationActor)
-	{
-		MoveReq.SetGoalActor(InDestinationActor);
-	}
-	else
-	{
-		MoveReq.SetGoalLocation(InPath.IsValid() && InPath->GetPathPoints().Num() ? InPath->GetPathPoints().Last().Location : FAISystem::InvalidLocation);
-	}
-
-	MoveReq.SetAcceptanceRadius(InAcceptanceRadius);
-	MoveReq.SetReachTestIncludesAgentRadius(bInStopOnOverlap);
-	MoveReq.SetUserData(InGameData);
-
-	return RequestMove(MoveReq, InPath);
-}
-
-void UPathFollowingComponent::AbortMove(const FString& Reason, FAIRequestID RequestID, bool bResetVelocity, bool bSilent, uint8 MessageFlags)
-{
-	const FPathFollowingResultFlags::Type AbortDetails =
-		(MessageFlags == EPathFollowingMessage::NoPath) ? FPathFollowingResultFlags::InvalidPath :
-		(MessageFlags == EPathFollowingMessage::OtherRequest) ? FPathFollowingResultFlags::NewRequest :
-		FPathFollowingResultFlags::None;
-
-	UE_VLOG(GetOwner(), LogPathFollowing, Log, TEXT("AbortMove with reason(%s)"), *Reason);
-	AbortMove(*this, AbortDetails, RequestID, bResetVelocity ? EPathFollowingVelocityMode::Reset : EPathFollowingVelocityMode::Keep);
-}
-
-void UPathFollowingComponent::PauseMove(FAIRequestID RequestID, bool bResetVelocity)
-{
-	PauseMove(RequestID, bResetVelocity ? EPathFollowingVelocityMode::Reset : EPathFollowingVelocityMode::Keep);
-}
-
 EPathFollowingAction::Type UPathFollowingComponent::GetPathActionType() const
 {
 	switch (Status)
@@ -1946,19 +1878,4 @@ FVector UPathFollowingComponent::GetPathDestination() const
 	return Path.IsValid() ? Path->GetDestinationLocation() : FVector::ZeroVector;
 }
 
-void UPathFollowingComponent::OnPathFinished(EPathFollowingResult::Type Result)
-{
-	OnPathFinished(Result, FPathFollowingResultFlags::None);
-}
-
-bool UPathFollowingComponent::HasReached(const FVector& TestPoint, float InAcceptanceRadius, bool bExactSpot) const
-{
-	return HasReached(TestPoint, bExactSpot ? EPathFollowingReachMode::ExactLocation : EPathFollowingReachMode::OverlapAgent, InAcceptanceRadius);
-}
-
-bool UPathFollowingComponent::HasReached(const AActor& TestGoal, float InAcceptanceRadius, bool bExactSpot, bool bUseNavAgentGoalLocation) const
-{
-	return HasReached(TestGoal, bExactSpot ? EPathFollowingReachMode::ExactLocation : EPathFollowingReachMode::OverlapAgentAndGoal, InAcceptanceRadius, bUseNavAgentGoalLocation);
-}
-
 #undef SHIPPING_STATIC
diff --git a/Engine/Source/Runtime/AIModule/Private/Perception/AIPerceptionComponent.cpp b/Engine/Source/Runtime/AIModule/Private/Perception/AIPerceptionComponent.cpp
--- a/Engine/Source/Runtime/AIModule/Private/Perception/AIPerceptionComponent.cpp
+++ b/Engine/Source/Runtime/AIModule/Private/Perception/AIPerceptionComponent.cpp
@@ -877,11 +877,6 @@ void UAIPerceptionComponent::DescribeSelfToVisLog(FVisualLogEntry* Snapshot) con
 //----------------------------------------------------------------------//
 // deprecated
 //----------------------------------------------------------------------//
-void UAIPerceptionComponent::UpdatePerceptionFilter(FAISenseID Channel, bool bNewValue)
-{
-	UpdatePerceptionAllowList(Channel, bNewValue);
-}
-
 void UAIPerceptionComponent::GetPerceivedActors(TSubclassOf<UAISense> SenseToUse, TArray<AActor*>& OutActors) const
 {
 	GetCurrentlyPerceivedActors(SenseToUse, OutActors);
diff --git a/Engine/Source/Runtime/AIModule/Public/BlueprintNodeHelpers.h b/Engine/Source/Runtime/AIModule/Public/BlueprintNodeHelpers.h
--- a/Engine/Source/Runtime/AIModule/Public/BlueprintNodeHelpers.h
+++ b/Engine/Source/Runtime/AIModule/Public/BlueprintNodeHelpers.h
@@ -43,10 +43,4 @@ namespace BlueprintNodeHelpers
 	{
 		return NodeObject.GetClass()->GetName().LeftChop(2);
 	}
-
-	//----------------------------------------------------------------------//
-	// DEPRECATED
-	//----------------------------------------------------------------------//
-	UE_DEPRECATED(4.11, "This version of HasBlueprintFunction is deprecated. Please use the one taking reference to UObject and StopAtClass rather than a pointers.")
-	AIMODULE_API bool HasBlueprintFunction(FName FuncName, const UObject* Object, const UClass* StopAtClass);
 }
diff --git a/Engine/Source/Runtime/AnimGraphRuntime/Private/BoneControllers/AnimNode_SkeletalControlBase.cpp b/Engine/Source/Runtime/AnimGraphRuntime/Private/BoneControllers/AnimNode_SkeletalControlBase.cpp
--- a/Engine/Source/Runtime/AnimGraphRuntime/Private/BoneControllers/AnimNode_SkeletalControlBase.cpp
+++ b/Engine/Source/Runtime/AnimGraphRuntime/Private/BoneControllers/AnimNode_SkeletalControlBase.cpp
@@ -146,11 +146,6 @@ void FAnimNode_SkeletalControlBase::AddDebugNodeData(FString& OutDebugData)
 
 void FAnimNode_SkeletalControlBase::EvaluateSkeletalControl_AnyThread(FComponentSpacePoseContext& Output, TArray<FBoneTransform>& OutBoneTransforms)
 {
-	DECLARE_SCOPE_HIERARCHICAL_COUNTER_ANIMNODE(EvaluateSkeletalControl_AnyThread)
-	PRAGMA_DISABLE_DEPRECATION_WARNINGS
-	// Call legacy implementation for backwards compatibility
-	EvaluateBoneTransforms(Output.AnimInstanceProxy->GetSkelMeshComponent(), Output.Pose, OutBoneTransforms);
-	PRAGMA_ENABLE_DEPRECATION_WARNINGS
 }
 
 void FAnimNode_SkeletalControlBase::SetAlpha(float InAlpha)
diff --git a/Engine/Source/Runtime/AnimGraphRuntime/Public/BoneControllers/AnimNode_SkeletalControlBase.h b/Engine/Source/Runtime/AnimGraphRuntime/Public/BoneControllers/AnimNode_SkeletalControlBase.h
--- a/Engine/Source/Runtime/AnimGraphRuntime/Public/BoneControllers/AnimNode_SkeletalControlBase.h
+++ b/Engine/Source/Runtime/AnimGraphRuntime/Public/BoneControllers/AnimNode_SkeletalControlBase.h
@@ -104,8 +104,6 @@ protected:
 
 	// use this function to evaluate for skeletal control base
 	virtual void EvaluateComponentSpaceInternal(FComponentSpacePoseContext& Context);
-	UE_DEPRECATED(4.16, "Please use EvaluateSkeletalControl_AnyThread.")
-	virtual void EvaluateBoneTransforms(USkeletalMeshComponent* SkelComp, FCSPose<FCompactPose>& MeshBases, TArray<FBoneTransform>& OutBoneTransforms) {}
 	// Evaluate the new component-space transforms for the affected bones.
 	virtual void EvaluateSkeletalControl_AnyThread(FComponentSpacePoseContext& Output, TArray<FBoneTransform>& OutBoneTransforms);
 	// return true if it is valid to Evaluate
diff --git a/Engine/Source/Runtime/ApplicationCore/Private/Linux/LinuxWindow.cpp b/Engine/Source/Runtime/ApplicationCore/Private/Linux/LinuxWindow.cpp
--- a/Engine/Source/Runtime/ApplicationCore/Private/Linux/LinuxWindow.cpp
+++ b/Engine/Source/Runtime/ApplicationCore/Private/Linux/LinuxWindow.cpp
@@ -826,11 +826,6 @@ bool FLinuxWindow::IsUtilityWindow() const
 	return bIsUtilityWindow;
 }
 
-bool FLinuxWindow::IsActivateWhenFirstShown() const
-{
-	return GetActivationPolicy() != EWindowActivationPolicy::Never;
-}
-
 EWindowActivationPolicy FLinuxWindow::GetActivationPolicy() const
 {
 	return Definition->ActivationPolicy;
diff --git a/Engine/Source/Runtime/ApplicationCore/Public/GenericPlatform/ICursor.h b/Engine/Source/Runtime/ApplicationCore/Public/GenericPlatform/ICursor.h
--- a/Engine/Source/Runtime/ApplicationCore/Public/GenericPlatform/ICursor.h
+++ b/Engine/Source/Runtime/ApplicationCore/Public/GenericPlatform/ICursor.h
@@ -111,7 +111,4 @@ public:
 	 * Allows overriding the shape of a particular cursor.
 	 */
 	virtual void SetTypeShape(EMouseCursor::Type InCursorType, void* CursorHandle) = 0;
-
-	UE_DEPRECATED(4.16, "Use SetTypeShape instead.")
-	void SetCustomShape(void* CursorHandle) { SetTypeShape(EMouseCursor::Custom, CursorHandle); }
 };
diff --git a/Engine/Source/Runtime/ApplicationCore/Public/Linux/LinuxWindow.h b/Engine/Source/Runtime/ApplicationCore/Public/Linux/LinuxWindow.h
--- a/Engine/Source/Runtime/ApplicationCore/Public/Linux/LinuxWindow.h
+++ b/Engine/Source/Runtime/ApplicationCore/Public/Linux/LinuxWindow.h
@@ -67,9 +67,6 @@ public:
 	/** TODO: describe */
 	bool IsUtilityWindow() const;
 
-	UE_DEPRECATED(4.16, "IsActivateWhenFirstShown() is deprecated. Please use GetActivationPolicy() instead")
-	bool IsActivateWhenFirstShown() const;
-
 	/** @return the window activation policy used when showing the window */
 	EWindowActivationPolicy GetActivationPolicy() const;
 	
diff --git a/Engine/Source/Runtime/AssetRegistry/Private/AssetRegistry.h b/Engine/Source/Runtime/AssetRegistry/Private/AssetRegistry.h
--- a/Engine/Source/Runtime/AssetRegistry/Private/AssetRegistry.h
+++ b/Engine/Source/Runtime/AssetRegistry/Private/AssetRegistry.h
@@ -152,9 +152,6 @@ public:
 
 	virtual void Tick (float DeltaTime) override;
 
-	UE_DEPRECATED(4.17, "IsUsingWorldAssets is now always true, remove any code that assumes it could be false")
-	static bool IsUsingWorldAssets() { return true; }
-
 	virtual void ReadLockEnumerateTagToAssetDatas(TFunctionRef<void(FName TagName, const TArray<const FAssetData*>& Assets)> Callback) const override;
 
 	virtual bool IsPathBeautificationNeeded(const FString& InAssetPath) const override;
diff --git a/Engine/Source/Runtime/Core/Public/Algo/IsSorted.h b/Engine/Source/Runtime/Core/Public/Algo/IsSorted.h
--- a/Engine/Source/Runtime/Core/Public/Algo/IsSorted.h
+++ b/Engine/Source/Runtime/Core/Public/Algo/IsSorted.h
@@ -55,37 +55,6 @@ namespace AlgoImpl
 
 namespace Algo
 {
-	/**
-	 * Tests if a range is sorted by its element type's operator<.
-	 *
-	 * @param  Array      A pointer to the array to test for being sorted.
-	 * @param  ArraySize  The number of elements in the array.
-	 *
-	 * @return true if the range is sorted, false otherwise.
-	 */
-	template <typename T>
-	UE_DEPRECATED(4.16, "IsSorted taking a pointer and size has been deprecated - please pass a TArrayView instead")
-	FORCEINLINE bool IsSorted(const T* Array, int32 ArraySize)
-	{
-		return AlgoImpl::IsSortedBy(Array, ArraySize, FIdentityFunctor(), TLess<>());
-	}
-
-	/**
-	 * Tests if a range is sorted by a user-defined predicate.
-	 *
-	 * @param  Array      A pointer to the array to test for being sorted.
-	 * @param  ArraySize  The number of elements in the array.
-	 * @param  Pred       A binary sorting predicate which describes the ordering of the elements in the array.
-	 *
-	 * @return true if the range is sorted, false otherwise.
-	 */
-	template <typename T, typename PredType>
-	UE_DEPRECATED(4.16, "IsSorted taking a pointer and size has been deprecated - please pass a TArrayView instead")
-	FORCEINLINE bool IsSorted(const T* Array, int32 ArraySize, PredType Pred)
-	{
-		return AlgoImpl::IsSortedBy(Array, ArraySize, FIdentityFunctor(), MoveTemp(Pred));
-	}
-
 	/**
 	 * Tests if a range is sorted by its element type's operator<.
 	 *
diff --git a/Engine/Source/Runtime/Core/Public/Math/Box2D.h b/Engine/Source/Runtime/Core/Public/Math/Box2D.h
--- a/Engine/Source/Runtime/Core/Public/Math/Box2D.h
+++ b/Engine/Source/Runtime/Core/Public/Math/Box2D.h
@@ -35,17 +35,6 @@ public:
 	/** Default constructor (no initialization). */
 	TBox2<T>() { }
 
-	/**
-	 * Creates and initializes a new box.
-	 *
-	 * The box extents are initialized to zero and the box is marked as invalid.
-	 */
-	UE_DEPRECATED(4.16, "Use ForceInit constructor instead.")
-	TBox2<T>( int32 )
-	{
-		Init();
-	}
-
 	/**
 	 * Creates and initializes a new box.
 	 *
diff --git a/Engine/Source/Runtime/Core/Public/Math/UnrealMathUtility.h b/Engine/Source/Runtime/Core/Public/Math/UnrealMathUtility.h
--- a/Engine/Source/Runtime/Core/Public/Math/UnrealMathUtility.h
+++ b/Engine/Source/Runtime/Core/Public/Math/UnrealMathUtility.h
@@ -770,12 +770,6 @@ public:
 		return Delta;
 	}
 
-	UE_DEPRECATED(4.12, "Please use FindDeltaAngleRadians(float A1, float A2) instead of FindDeltaAngle(float A1, float A2).")
-	UE_NODISCARD static constexpr float FindDeltaAngle(float A1, float A2)
-	{
-		return FindDeltaAngleRadians(A1, A2);
-	}
-
 	/** Given a heading which may be outside the +/- PI range, 'unwind' it back into that range. */
 	template<typename T, TEMPLATE_REQUIRES(TIsFloatingPoint<T>::Value)>
 	UE_NODISCARD static constexpr T UnwindRadians(T A)
diff --git a/Engine/Source/Runtime/Core/Public/Misc/AutomationTest.h b/Engine/Source/Runtime/Core/Public/Misc/AutomationTest.h
--- a/Engine/Source/Runtime/Core/Public/Misc/AutomationTest.h
+++ b/Engine/Source/Runtime/Core/Public/Misc/AutomationTest.h
@@ -1245,12 +1245,6 @@ public:
 	 */
 	virtual void AddWarning( const FString& InWarning, int32 StackOffset = 0);
 
-	UE_DEPRECATED(4.16, "Use AddInfo")
-	FORCEINLINE void AddLogItem(const FString& InLogItem)
-	{
-		AddInfo(InLogItem, 0);
-	}
-
 	/**
 	 * Adds a log item to this test
 	 *
diff --git a/Engine/Source/Runtime/Core/Public/Misc/ConfigCacheIni.h b/Engine/Source/Runtime/Core/Public/Misc/ConfigCacheIni.h
--- a/Engine/Source/Runtime/Core/Public/Misc/ConfigCacheIni.h
+++ b/Engine/Source/Runtime/Core/Public/Misc/ConfigCacheIni.h
@@ -178,11 +178,6 @@ public:
 		bRead = InBRead;
 	}
 #endif
-	UE_DEPRECATED(4.12, "Please switch to explicitly doing a GetValue() or GetSavedValue()")
-	operator const FString& () const { return GetValue(); }
-
-	UE_DEPRECATED(4.12, "Please switch to explicitly doing a GetValue() or GetSavedValue()")
-	const TCHAR* operator*() const { return *GetValue(); }
 
 	bool operator==(const FConfigValue& Other) const { return (SavedValue.Compare(Other.SavedValue) == 0); }
 	bool operator!=(const FConfigValue& Other) const { return !(FConfigValue::operator==(Other)); }
diff --git a/Engine/Source/Runtime/CoreUObject/Private/Misc/PackageName.cpp b/Engine/Source/Runtime/CoreUObject/Private/Misc/PackageName.cpp
--- a/Engine/Source/Runtime/CoreUObject/Private/Misc/PackageName.cpp
+++ b/Engine/Source/Runtime/CoreUObject/Private/Misc/PackageName.cpp
@@ -791,23 +791,6 @@ bool FPackageName::TryConvertLongPackageNameToFilename(const FString& InLongPack
 	return false;
 }
 
-bool FPackageName::ConvertRootPathToContentPath( const FString& RootPath, FString& OutContentPath)
-{
-	const auto& Paths = FLongPackagePathsSingleton::Get();
-	FReadScopeLock ScopeLock(ContentMountPointCriticalSection);
-	for (const auto& Pair : Paths.ContentRootToPath)
-	{
-		if (FPathViews::IsParentPathOf(Pair.RootPath, RootPath))
-		{
-			OutContentPath = Pair.ContentPath;
-			return true;
-		}
-	}
-
-	// This is not a long package name or the root folder is not handled in the above cases
-	return false;
-}
-
 FString FPackageName::LongPackageNameToFilename(const FString& InLongPackageName, const FString& InExtension)
 {
 	FString Result;
@@ -2053,20 +2036,6 @@ const FString& FPackageName::GetTextMapPackageExtension()
 	return TextMapPackageExtension;
 }
 
-FString FPackageName::PackageFromPath(const TCHAR* InPathName)
-{
-	FString PackageName;
-	if (FPackageName::TryConvertFilenameToLongPackageName(InPathName, PackageName))
-	{
-		return PackageName;
-	}
-	else
-	{
-		// Not a valid package filename
-		return InPathName;
-	}
-}
-
 bool FPackageName::IsTextPackageExtension(const TCHAR* Ext)
 {
 	return IsTextAssetPackageExtension(Ext) || IsTextMapPackageExtension(Ext);
diff --git a/Engine/Source/Runtime/CoreUObject/Public/Misc/PackageName.h b/Engine/Source/Runtime/CoreUObject/Public/Misc/PackageName.h
--- a/Engine/Source/Runtime/CoreUObject/Public/Misc/PackageName.h
+++ b/Engine/Source/Runtime/CoreUObject/Public/Misc/PackageName.h
@@ -710,12 +710,6 @@ public:
 	 */
 	static bool FixPackageNameCase(FString& LongPackageName, FStringView Extension);
 
-	UE_DEPRECATED(4.17, "Deprecated. Call TryConvertLongPackageNameToFilename instead, which also works on nested paths")
-	static bool ConvertRootPathToContentPath(const FString& RootPath, FString& OutContentPath);
-
-	UE_DEPRECATED(4.17, "Deprecated. Call TryConvertFilenameToLongPackageName instead")
-	static FString PackageFromPath(const TCHAR* InPathName);
-
 	/** Override whether a package exist or not. */
 	DECLARE_DELEGATE_RetVal_OneParam(bool, FDoesPackageExistOverride, FName);
 	static FDoesPackageExistOverride& DoesPackageExistOverride()
diff --git a/Engine/Source/Runtime/CoreUObject/Public/UObject/Class.h b/Engine/Source/Runtime/CoreUObject/Public/UObject/Class.h
--- a/Engine/Source/Runtime/CoreUObject/Public/UObject/Class.h
+++ b/Engine/Source/Runtime/CoreUObject/Public/UObject/Class.h
@@ -2347,9 +2347,6 @@ public:
 	 * @return The tooltip for this object.
 	 */
 	FText GetToolTipTextByIndex(int32 NameIndex) const;
-
-	UE_DEPRECATED(4.16, "GetToolTipText with name index is deprecated, call GetToolTipTextByIndex instead")
-	FText GetToolTipText(int32 NameIndex) const { return GetToolTipTextByIndex(NameIndex); }
 #endif
 
 #if WITH_EDITORONLY_DATA
@@ -2526,31 +2523,6 @@ public:
 		out_TextValue = GetDisplayValueAsText(EnumeratorValue);
 	}
 
-	// Deprecated Functions
-	UE_DEPRECATED(4.16, "FindEnumIndex is deprecated, call GetIndexByName or GetValueByName instead")
-	int32 FindEnumIndex(FName InName) const { return GetIndexByName(InName, EGetByNameFlags::ErrorIfNotFound); }
-
-	UE_DEPRECATED(4.16, "FindEnumRedirects is deprecated, call GetIndexByNameString instead")
-	static int32 FindEnumRedirects(const UEnum* Enum, FName EnumEntryName) { return Enum->GetIndexByNameString(EnumEntryName.ToString()); }
-
-	UE_DEPRECATED(4.16, "GetEnum is deprecated, call GetNameByIndex instead")
-	FName GetEnum(int32 InIndex) const { return GetNameByIndex(InIndex); }
-
-	UE_DEPRECATED(4.16, "GetEnumNameStringByValue is deprecated, call GetNameStringByValue instead")
-	FString GetEnumNameStringByValue(int64 InValue) const { return GetNameStringByValue(InValue); }
-
-	UE_DEPRECATED(4.16, "GetEnumName is deprecated, call GetNameStringByIndex instead")
-	FString GetEnumName(int32 InIndex) const { return GetNameStringByIndex(InIndex); }
-
-	UE_DEPRECATED(4.16, "GetDisplayNameText with name index is deprecated, call GetDisplayNameTextByIndex instead")
-	FText GetDisplayNameText(int32 NameIndex) const { return GetDisplayNameTextByIndex(NameIndex); }
-
-	UE_DEPRECATED(4.16, "GetEnumText with name index is deprecated, call GetDisplayNameTextByIndex instead")
-	FText GetEnumText(int32 NameIndex) const { return GetDisplayNameTextByIndex(NameIndex); }
-
-	UE_DEPRECATED(4.16, "GetEnumTextByValue with name index is deprecated, call GetDisplayNameTextByValue instead")
-	FText GetEnumTextByValue(int64 Value) { return GetDisplayNameTextByValue(Value);  }
-
 	// UObject interface.
 	virtual void Serialize(FArchive& Ar) override;
 	virtual void BeginDestroy() override;
diff --git a/Engine/Source/Runtime/Engine/Classes/Animation/AnimInstance.h b/Engine/Source/Runtime/Engine/Classes/Animation/AnimInstance.h
--- a/Engine/Source/Runtime/Engine/Classes/Animation/AnimInstance.h
+++ b/Engine/Source/Runtime/Engine/Classes/Animation/AnimInstance.h
@@ -388,43 +388,6 @@ class ENGINE_API UAnimInstance : public UObject
 	UPROPERTY(Category = RootMotion, EditDefaultsOnly)
 	TEnumAsByte<ERootMotionMode::Type> RootMotionMode;
 
-#if WITH_EDITORONLY_DATA
-	/** DeltaTime **/
-	UPROPERTY()
-	float DeltaTime_DEPRECATED;
-
-	/** 
-	 * DEPRECATED: No longer used.
-	 * Allows this anim instance to update its native update, blend tree, montages and asset players on
-	 * a worker thread. this requires certain conditions to be met:
-	 * - All access of variables in the blend tree should be a direct access of a member variable
-	 * - No BlueprintUpdateAnimation event should be used (i.e. the event graph should be empty). Only native update is permitted.
-	 */
-	UE_DEPRECATED(4.15, "This variable is no longer used. Use bUseMultiThreadedAnimationUpdate on the UAnimBlueprint to control this.")
-	UPROPERTY()
-	uint8 bRunUpdatesInWorkerThreads_DEPRECATED : 1;
-
-	/** 
-	 * DEPRECATED: No longer used.
-	 * Whether we can use parallel updates for our animations.
-	 * Conditions affecting this include:
-	 * - Use of BlueprintUpdateAnimation
-	 * - Use of non 'fast-path' EvaluateGraphExposedInputs in the node graph
-	 */
-	UE_DEPRECATED(4.15, "This variable is no longer used. Use bUseMultiThreadedAnimationUpdate on the UAnimBlueprint to control this.")
-	UPROPERTY()
-	uint8 bCanUseParallelUpdateAnimation_DEPRECATED : 1;
-
-
-	/**
-	 * Selecting this option will cause the compiler to emit warnings whenever a call into Blueprint
-	 * is made from the animation graph. This can help track down optimizations that need to be made.
-	 */
-	UE_DEPRECATED(4.15, "This variable is no longer used. Use bWarnAboutBlueprintUsage on the UAnimBlueprint to control this.")
-	UPROPERTY()
-	uint8 bWarnAboutBlueprintUsage_DEPRECATED : 1;
-#endif
-
 	/**
 	 * Allows this anim instance to update its native update, blend tree, montages and asset players on
 	 * a worker thread. This flag is propagated from the UAnimBlueprint to this instance by the compiler.
@@ -789,10 +752,6 @@ public:
 		Note that there might be multiple Active at the same time. This will only return the first active one it finds. **/
 	FAnimMontageInstance* GetActiveMontageInstance() const;
 
-	/** Get Active FAnimMontageInstance for given Montage asset. Will return NULL if Montage is not currently Active. */
-	UE_DEPRECATED(4.13, "Please use GetActiveInstanceForMontage(const UAnimMontage* Montage)")
-	FAnimMontageInstance* GetActiveInstanceForMontage(UAnimMontage const& Montage) const;
-
 	/** Get Active FAnimMontageInstance for given Montage asset. Will return NULL if Montage is not currently Active. */
 	FAnimMontageInstance* GetActiveInstanceForMontage(const UAnimMontage* Montage) const;
 
@@ -1284,11 +1243,6 @@ public:
 	// Native thread safe update override point. Executed on a worker thread just prior to graph update 
 	// for linked anim instances, only called when the hosting node(s) are relevant
 	virtual void NativeThreadSafeUpdateAnimation(float DeltaSeconds);
-	// Native update override point. Can be called from a worker thread. This is a good place to do any
-	// heavy lifting (as opposed to NativeUpdateAnimation_GameThread()).
-	// This function should not be used. Worker thread updates should be performed in the FAnimInstanceProxy attached to this instance.
-	UE_DEPRECATED(4.15, "This function is only called for backwards-compatibility. It is no longer called on a worker thread.")
-	virtual void NativeUpdateAnimation_WorkerThread(float DeltaSeconds);
 	// Native Post Evaluate override point
 	virtual void NativePostEvaluateAnimation();
 	// Native Uninitialize override point
diff --git a/Engine/Source/Runtime/Engine/Classes/Animation/AnimNodeBase.h b/Engine/Source/Runtime/Engine/Classes/Animation/AnimNodeBase.h
--- a/Engine/Source/Runtime/Engine/Classes/Animation/AnimNodeBase.h
+++ b/Engine/Source/Runtime/Engine/Classes/Animation/AnimNodeBase.h
@@ -857,16 +857,6 @@ struct ENGINE_API FAnimNode_Base
 	virtual ~FAnimNode_Base() {}
 
 	/** Deprecated functions */
-	UE_DEPRECATED(4.17, "Please use Initialize_AnyThread instead")
-	virtual void Initialize(const FAnimationInitializeContext& Context);
-	UE_DEPRECATED(4.17, "Please use CacheBones_AnyThread instead")
-	virtual void CacheBones(const FAnimationCacheBonesContext& Context) {}
-	UE_DEPRECATED(4.17, "Please use Update_AnyThread instead")
-	virtual void Update(const FAnimationUpdateContext& Context) {}
-	UE_DEPRECATED(4.17, "Please use Evaluate_AnyThread instead")
-	virtual void Evaluate(FPoseContext& Output) { check(false); }
-	UE_DEPRECATED(4.17, "Please use EvaluateComponentSpace_AnyThread instead")
-	virtual void EvaluateComponentSpace(FComponentSpacePoseContext& Output) { check(false); }
 	UE_DEPRECATED(4.20, "Please use ResetDynamics with an ETeleportPhysics flag instead")
 	virtual void ResetDynamics() {}
 	UE_DEPRECATED(5.0, "Please use IGraphMessage instead")
@@ -962,10 +952,6 @@ protected:
 	/** Get the LOD threshold at which this node is enabled. Node is enabled if the current LOD >= threshold. */
 	virtual int32 GetLODThreshold() const { return INDEX_NONE; }
 
-	/** Deprecated function */
-	UE_DEPRECATED(4.17, "Please use OnInitializeAnimInstance instead")
-	virtual void RootInitialize(const FAnimInstanceProxy* InProxy) {}
-
 	/** Called once, from game thread as the parent anim instance is created */
 	virtual void OnInitializeAnimInstance(const FAnimInstanceProxy* InProxy, const UAnimInstance* InAnimInstance);
 
diff --git a/Engine/Source/Runtime/Engine/Classes/Animation/AnimationAsset.h b/Engine/Source/Runtime/Engine/Classes/Animation/AnimationAsset.h
--- a/Engine/Source/Runtime/Engine/Classes/Animation/AnimationAsset.h
+++ b/Engine/Source/Runtime/Engine/Classes/Animation/AnimationAsset.h
@@ -629,21 +629,6 @@ struct FRootMotionMovementParams
 private:
 	ENGINE_API static FVector RootMotionScale;
 
-	// TODO: Remove when we make RootMotionTransform private
-	FORCEINLINE FTransform& GetRootMotionTransformInternal()
-	{
-		PRAGMA_DISABLE_DEPRECATION_WARNINGS
-		return RootMotionTransform;
-		PRAGMA_ENABLE_DEPRECATION_WARNINGS
-	}
-
-	FORCEINLINE const FTransform& GetRootMotionTransformInternal() const
-	{
-		PRAGMA_DISABLE_DEPRECATION_WARNINGS
-		return RootMotionTransform;
-		PRAGMA_ENABLE_DEPRECATION_WARNINGS
-	}
-
 public:
 	
 	UPROPERTY()
@@ -652,11 +637,11 @@ public:
 	UPROPERTY()
 	float BlendWeight;
 
-	// To be made private
-	UE_DEPRECATED(4.13, "RootMotionTransform should not be accessed directly, please use GetRootMotionTransform() to read it or one of the set/accumulate functions to modify it")
+private:
 	UPROPERTY()
 	FTransform RootMotionTransform;
 
+public:
 	FRootMotionMovementParams()
 		: bHasRootMotion(false)
 		, BlendWeight(0.f)
@@ -669,29 +654,29 @@ public:
 		: bHasRootMotion(Other.bHasRootMotion)
 		, BlendWeight(Other.BlendWeight)
 	{
-		GetRootMotionTransformInternal() = Other.GetRootMotionTransformInternal();
+		RootMotionTransform = Other.RootMotionTransform;
 	}
 
 	FRootMotionMovementParams(const FRootMotionMovementParams&& Other)
 		: bHasRootMotion(Other.bHasRootMotion)
 		, BlendWeight(Other.BlendWeight)
 	{
-		GetRootMotionTransformInternal() = Other.GetRootMotionTransformInternal();
+		RootMotionTransform = Other.RootMotionTransform;
 	}
 
 	FRootMotionMovementParams& operator=(const FRootMotionMovementParams& Other)
 	{
 		bHasRootMotion = Other.bHasRootMotion;
 		BlendWeight = Other.BlendWeight;
-		GetRootMotionTransformInternal() = Other.GetRootMotionTransformInternal();
+		RootMotionTransform = Other.RootMotionTransform;
 		return *this;
 	}
 
 	void Set(const FTransform& InTransform)
 	{
 		bHasRootMotion = true;
-		GetRootMotionTransformInternal() = InTransform;
-		GetRootMotionTransformInternal().SetScale3D(RootMotionScale);
+		RootMotionTransform = InTransform;
+		RootMotionTransform.SetScale3D(RootMotionScale);
 		BlendWeight = 1.f;
 	}
 
@@ -703,8 +688,8 @@ public:
 		}
 		else
 		{
-			GetRootMotionTransformInternal() = InTransform * GetRootMotionTransformInternal();
-			GetRootMotionTransformInternal().SetScale3D(RootMotionScale);
+			RootMotionTransform = InTransform * RootMotionTransform;
+			RootMotionTransform.SetScale3D(RootMotionScale);
 		}
 	}
 
@@ -712,7 +697,7 @@ public:
 	{
 		if (MovementParams.bHasRootMotion)
 		{
-			Accumulate(MovementParams.GetRootMotionTransformInternal());
+			Accumulate(MovementParams.RootMotionTransform);
 		}
 	}
 
@@ -721,8 +706,8 @@ public:
 		const ScalarRegister VBlendWeight(InBlendWeight);
 		if (bHasRootMotion)
 		{
-			GetRootMotionTransformInternal().AccumulateWithShortestRotation(InTransform, VBlendWeight);
-			GetRootMotionTransformInternal().SetScale3D(RootMotionScale);
+			RootMotionTransform.AccumulateWithShortestRotation(InTransform, VBlendWeight);
+			RootMotionTransform.SetScale3D(RootMotionScale);
 			BlendWeight += InBlendWeight;
 		}
 		else
@@ -736,7 +721,7 @@ public:
 	{
 		if (MovementParams.bHasRootMotion)
 		{
-			AccumulateWithBlend(MovementParams.GetRootMotionTransformInternal(), InBlendWeight);
+			AccumulateWithBlend(MovementParams.RootMotionTransform, InBlendWeight);
 		}
 	}
 
@@ -753,28 +738,28 @@ public:
 		{
 			AccumulateWithBlend(FTransform(), WeightLeft);
 		}
-		GetRootMotionTransformInternal().NormalizeRotation();
+		RootMotionTransform.NormalizeRotation();
 	}
 
 	FRootMotionMovementParams ConsumeRootMotion(float Alpha)
 	{
 		const ScalarRegister VAlpha(Alpha);
-		FTransform PartialRootMotion = (GetRootMotionTransformInternal()*VAlpha);
+		FTransform PartialRootMotion = (RootMotionTransform*VAlpha);
 		PartialRootMotion.SetScale3D(RootMotionScale); // Reset scale after multiplication above
 		PartialRootMotion.NormalizeRotation();
-		GetRootMotionTransformInternal() = GetRootMotionTransformInternal().GetRelativeTransform(PartialRootMotion);
-		GetRootMotionTransformInternal().NormalizeRotation(); //Make sure we are normalized, this needs to be investigated further
+		RootMotionTransform = RootMotionTransform.GetRelativeTransform(PartialRootMotion);
+		RootMotionTransform.NormalizeRotation(); //Make sure we are normalized, this needs to be investigated further
 
 		FRootMotionMovementParams ReturnParams;
 		ReturnParams.Set(PartialRootMotion);
 
 		check(PartialRootMotion.IsRotationNormalized());
-		check(GetRootMotionTransformInternal().IsRotationNormalized());
+		check(RootMotionTransform.IsRotationNormalized());
 		return ReturnParams;
 	}
 
-	const FTransform& GetRootMotionTransform() const { return GetRootMotionTransformInternal(); }
-	void ScaleRootMotionTranslation(float TranslationScale) { GetRootMotionTransformInternal().ScaleTranslation(TranslationScale); }
+	const FTransform& GetRootMotionTransform() const { return RootMotionTransform; }
+	void ScaleRootMotionTranslation(float TranslationScale) { RootMotionTransform.ScaleTranslation(TranslationScale); }
 };
 
 // This structure is used to either advance or synchronize animation players
diff --git a/Engine/Source/Runtime/Engine/Classes/Animation/Skeleton.h b/Engine/Source/Runtime/Engine/Classes/Animation/Skeleton.h
--- a/Engine/Source/Runtime/Engine/Classes/Animation/Skeleton.h
+++ b/Engine/Source/Runtime/Engine/Classes/Animation/Skeleton.h
@@ -948,17 +948,6 @@ public:
 	 */
 	ENGINE_API bool RecreateBoneTree(USkeletalMesh* InSkelMesh);
 
-	/** This is const accessor for BoneTree
-	 *  Understand there will be a lot of need to access BoneTree, but 
-	 *	Anybody modifying BoneTree will corrupt animation data, so will need to make sure it's not modifiable outside of Skeleton
-	 *	You can add new BoneNode but you can't modify current list. The index will be referenced by Animation data.
-	 */
-	UE_DEPRECATED(4.14, "GetBoneTree should not be called. Please use GetBoneTranslationRetargetingMode()/SetBoneTranslationRetargetingMode() functions instead")
-	const TArray<FBoneNode>& GetBoneTree()	const
-	{ 
-		return BoneTree;	
-	}
-
 	// @todo document
 	const TArray<FTransform>& GetRefLocalPoses( FName RetargetSource = NAME_None ) const 
 	{
diff --git a/Engine/Source/Runtime/Engine/Classes/EdGraph/EdGraph.h b/Engine/Source/Runtime/Engine/Classes/EdGraph/EdGraph.h
--- a/Engine/Source/Runtime/Engine/Classes/EdGraph/EdGraph.h
+++ b/Engine/Source/Runtime/Engine/Classes/EdGraph/EdGraph.h
@@ -130,17 +130,6 @@ public:
 #endif
 
 public:
-	/** 
-	 * Creates an empty node of the given type. 
-	 * User is fully responsible for building the node. Most nodes should be created through
-	 * a FGraphNodeCreator builder.
-	 */
-	template <typename NodeClass>
-	UE_DEPRECATED(4.17, "Use CreateIntermediateNode instead.")
-	NodeClass* CreateBlankNode()
-	{
-		return CreateIntermediateNode<NodeClass>();
-	}
 
 	template <typename NodeClass>
 	NodeClass* CreateIntermediateNode()
diff --git a/Engine/Source/Runtime/Engine/Classes/EdGraph/EdGraphNode.h b/Engine/Source/Runtime/Engine/Classes/EdGraph/EdGraphNode.h
--- a/Engine/Source/Runtime/Engine/Classes/EdGraph/EdGraphNode.h
+++ b/Engine/Source/Runtime/Engine/Classes/EdGraph/EdGraphNode.h
@@ -308,18 +308,13 @@ private:
 	UPROPERTY()
 	uint8 bUserSetEnabledState:1;
 
-protected:
-	/** (DEPRECATED) Value used for AllowSplitPins(). Do not override. */
-	uint8 bAllowSplitPins_DEPRECATED:1;
-
-private:
+#if WITH_EDITORONLY_DATA
 	/** (DEPRECATED) FALSE if the node is a disabled, which eliminates it from being compiled */
 	UPROPERTY()
 	uint8 bIsNodeEnabled_DEPRECATED:1;
 
 public:
 
-#if WITH_EDITORONLY_DATA
 	/** If true, this node can be resized and should be drawn with a resize handle */
 	UPROPERTY()
 	uint8 bCanResizeNode:1;
@@ -630,16 +625,7 @@ public:
 	virtual bool ShouldOverridePinNames() const { return false; }
 
 	/** Whether or not struct pins belonging to this node should be allowed to be split or not. */
-	UE_DEPRECATED(4.14, "Please call CanSplitPin and provide the specific Pin to split.")
-	virtual bool AllowSplitPins() const { return bAllowSplitPins_DEPRECATED; }
-
-	/** Whether or not struct pins belonging to this node should be allowed to be split or not. */
-	virtual bool CanSplitPin(const UEdGraphPin* Pin) const 
-	{
-		PRAGMA_DISABLE_DEPRECATION_WARNINGS
-		return AllowSplitPins();
-		PRAGMA_ENABLE_DEPRECATION_WARNINGS
-	}
+	virtual bool CanSplitPin(const UEdGraphPin* Pin) const { return false; }
 
 	/** Gets the overridden name for the specified pin, if any */
 	virtual FText GetPinNameOverride(const UEdGraphPin& Pin) const { return FText::GetEmpty(); }
@@ -789,10 +775,6 @@ public:
 	 */
 	virtual FString GetDocumentationExcerptName() const;
 
-	/** @return Icon to use in menu or on node */
-	UE_DEPRECATED(4.13, "Please override 'virtual FSlateIcon GetIconAndTint(FLinearColor& OutColor) const;' instead.")
-	virtual FName GetPaletteIcon(FLinearColor& OutColor) const { return NAME_None; }
-
 	/** @return Icon to use in menu or on node */
 	virtual FSlateIcon GetIconAndTint(FLinearColor& OutColor) const;
 
diff --git a/Engine/Source/Runtime/Engine/Classes/Engine/GameViewportClient.h b/Engine/Source/Runtime/Engine/Classes/Engine/GameViewportClient.h
--- a/Engine/Source/Runtime/Engine/Classes/Engine/GameViewportClient.h
+++ b/Engine/Source/Runtime/Engine/Classes/Engine/GameViewportClient.h
@@ -757,15 +757,6 @@ public:
 	 */
 	virtual bool CaptureMouseOnLaunch() override;
 
-	/**
-	 * Sets whether or not the cursor is locked to the viewport when the viewport captures the mouse
-	 */
-	UE_DEPRECATED(4.13, "Mouse locking is now controlled by an enum value. Please call UGameViewportClient::SetMouseLockMode(EMouseLockMode) instead.")
-	void SetLockDuringCapture(bool InLockDuringCapture)
-	{
-		SetMouseLockMode( InLockDuringCapture ? EMouseLockMode::LockOnCapture : EMouseLockMode::DoNotLock );
-	}
-
 	/**
 	 * Gets whether or not the cursor is locked to the viewport when the viewport captures the mouse
 	 */
diff --git a/Engine/Source/Runtime/Engine/Classes/GameFramework/GameState.h b/Engine/Source/Runtime/Engine/Classes/GameFramework/GameState.h
--- a/Engine/Source/Runtime/Engine/Classes/GameFramework/GameState.h
+++ b/Engine/Source/Runtime/Engine/Classes/GameFramework/GameState.h
@@ -67,9 +67,6 @@ public:
 	/** Called periodically, overridden by subclasses */
 	virtual void DefaultTimer();
 
-	UE_DEPRECATED(4.14, "ShouldShowGore is deprecated, it is not supported by newer assets")
-	virtual bool ShouldShowGore() const;
-
 	//~ Begin AActor Interface
 	virtual void PostInitializeComponents() override;
 	//~ End AActor Interface
diff --git a/Engine/Source/Runtime/Engine/Classes/GameFramework/PlayerController.h b/Engine/Source/Runtime/Engine/Classes/GameFramework/PlayerController.h
--- a/Engine/Source/Runtime/Engine/Classes/GameFramework/PlayerController.h
+++ b/Engine/Source/Runtime/Engine/Classes/GameFramework/PlayerController.h
@@ -1192,12 +1192,6 @@ public:
 		ClientPlayForceFeedback_Internal(ForceFeedbackEffect, Params);
 	}
 
-	UE_DEPRECATED(4.22, "Use version that specifies parameters using a struct instead of a list of parameters")
-	void ClientPlayForceFeedback(class UForceFeedbackEffect* ForceFeedbackEffect, bool bLooping, bool bIgnoreTimeDilation, FName Tag);
-
-	UE_DEPRECATED(4.18, "Use version that specifies parameters using a struct instead of a list of parameters")
-	void ClientPlayForceFeedback(class UForceFeedbackEffect* ForceFeedbackEffect, bool bLooping, FName Tag);
-
 	/** 
 	 * Stops a playing force feedback pattern
 	 * @param	ForceFeedbackEffect		If set only patterns from that effect will be stopped
diff --git a/Engine/Source/Runtime/Engine/Classes/Particles/EmitterCameraLensEffectBase.h b/Engine/Source/Runtime/Engine/Classes/Particles/EmitterCameraLensEffectBase.h
--- a/Engine/Source/Runtime/Engine/Classes/Particles/EmitterCameraLensEffectBase.h
+++ b/Engine/Source/Runtime/Engine/Classes/Particles/EmitterCameraLensEffectBase.h
@@ -26,10 +26,6 @@ protected:
 	UPROPERTY(EditDefaultsOnly, Category = EmitterCameraLensEffectBase)
 	TObjectPtr<class UParticleSystem> PS_CameraEffect;
 
-	/** The effect to use for non extreme content */
-	UPROPERTY()
-	TObjectPtr<class UParticleSystem> PS_CameraEffectNonExtremeContent_DEPRECATED;
-
 	/** Camera this emitter is attached to, will be notified when emitter is destroyed */
 	UPROPERTY(transient)
 	TObjectPtr<class APlayerCameraManager> BaseCamera;
diff --git a/Engine/Source/Runtime/Engine/Classes/PhysicsEngine/ConstraintInstance.h b/Engine/Source/Runtime/Engine/Classes/PhysicsEngine/ConstraintInstance.h
--- a/Engine/Source/Runtime/Engine/Classes/PhysicsEngine/ConstraintInstance.h
+++ b/Engine/Source/Runtime/Engine/Classes/PhysicsEngine/ConstraintInstance.h
@@ -740,12 +740,6 @@ public:
 	/** Get the linear drive's strength parameters */
 	void GetLinearDriveParams(float& OutPositionStrength, float& OutVelocityStrength, float& OutForceLimit);
 
-	UE_DEPRECATED(4.15, "Please call SetOrientationDriveTwistAndSwing. Note the order of bools is reversed (Make sure to pass Twist and then Swing)")
-	void SetAngularPositionDrive(bool bInEnableSwingDrive, bool bInEnableTwistDrive)
-	{
-		SetOrientationDriveTwistAndSwing(bInEnableTwistDrive, bInEnableSwingDrive);
-	}
-
 	/** Set which twist and swing orientation drives are enabled. Only applicable when Twist And Swing drive mode is used */
 	void SetOrientationDriveTwistAndSwing(bool bInEnableTwistDrive, bool bInEnableSwingDrive);
 
@@ -776,12 +770,6 @@ public:
 		return ProfileInstance.AngularDrive.OrientationTarget;
 	}
 
-	UE_DEPRECATED(4.15, "Please call SetAngularVelocityDriveTwistAndSwing. Note the order of bools is reversed (Make sure to pass Twist and then Swing)")
-	void SetAngularVelocityDrive(bool bInEnableSwingDrive, bool bInEnableTwistDrive)
-	{
-		SetAngularVelocityDriveTwistAndSwing(bInEnableTwistDrive, bInEnableSwingDrive);
-	}
-
 	/** Set which twist and swing angular velocity drives are enabled. Only applicable when Twist And Swing drive mode is used */
 	void SetAngularVelocityDriveTwistAndSwing(bool bInEnableTwistDrive, bool bInEnableSwingDrive);
 
diff --git a/Engine/Source/Runtime/Engine/Classes/Sound/SoundWaveStreaming.h b/Engine/Source/Runtime/Engine/Classes/Sound/SoundWaveStreaming.h
deleted file mode 100644
diff --git a/Engine/Source/Runtime/Engine/Private/Animation/AnimInstance.cpp b/Engine/Source/Runtime/Engine/Private/Animation/AnimInstance.cpp
--- a/Engine/Source/Runtime/Engine/Private/Animation/AnimInstance.cpp
+++ b/Engine/Source/Runtime/Engine/Private/Animation/AnimInstance.cpp
@@ -518,11 +518,6 @@ void UAnimInstance::UpdateAnimation(float DeltaSeconds, bool bNeedsValidRootMoti
 		SCOPE_CYCLE_COUNTER(STAT_NativeUpdateAnimation);
 		CSV_SCOPED_TIMING_STAT(Animation, NativeUpdate);
 		NativeUpdateAnimation(DeltaSeconds);
-
-PRAGMA_DISABLE_DEPRECATION_WARNINGS
-		// @todo: remove once deprecated - called for backwards-compatibility¬
-		NativeUpdateAnimation_WorkerThread(DeltaSeconds);
-PRAGMA_ENABLE_DEPRECATION_WARNINGS
 	}
 	{
 		SCOPE_CYCLE_COUNTER(STAT_BlueprintUpdateAnimation);
@@ -795,10 +790,6 @@ void UAnimInstance::NativeThreadSafeUpdateAnimation(float DeltaSeconds)
 {
 }
 
-void UAnimInstance::NativeUpdateAnimation_WorkerThread(float DeltaSeconds)
-{
-}
-
 void UAnimInstance::NativePostEvaluateAnimation()
 {
 }
@@ -3290,11 +3281,6 @@ UAnimInstance* UAnimInstance::GetLinkedAnimLayerInstanceByClass(TSubclassOf<UAni
 	return nullptr;
 }
 
-FAnimMontageInstance* UAnimInstance::GetActiveInstanceForMontage(UAnimMontage const& Montage) const
-{
-	return GetActiveInstanceForMontage(&Montage);
-}
-
 FAnimMontageInstance* UAnimInstance::GetActiveInstanceForMontage(const UAnimMontage* Montage) const
 {
 	FAnimMontageInstance* const* FoundInstancePtr = ActiveMontagesMap.Find(Montage);
diff --git a/Engine/Source/Runtime/Engine/Private/Animation/AnimNodeBase.cpp b/Engine/Source/Runtime/Engine/Private/Animation/AnimNodeBase.cpp
--- a/Engine/Source/Runtime/Engine/Private/Animation/AnimNodeBase.cpp
+++ b/Engine/Source/Runtime/Engine/Private/Animation/AnimNodeBase.cpp
@@ -81,50 +81,24 @@ void FComponentSpacePoseContext::ResetToRefPose()
 /////////////////////////////////////////////////////
 // FAnimNode_Base
 
-PRAGMA_DISABLE_DEPRECATION_WARNINGS
-void FAnimNode_Base::Initialize(const FAnimationInitializeContext& Context)
-{
-}
-PRAGMA_ENABLE_DEPRECATION_WARNINGS
-
 void FAnimNode_Base::Initialize_AnyThread(const FAnimationInitializeContext& Context)
 {
-	// Call legacy implementation for backwards compatibility
-	PRAGMA_DISABLE_DEPRECATION_WARNINGS
-	Initialize(Context);
-	PRAGMA_ENABLE_DEPRECATION_WARNINGS
 }
 
 void FAnimNode_Base::CacheBones_AnyThread(const FAnimationCacheBonesContext& Context)
 {
-	// Call legacy implementation for backwards compatibility
-	PRAGMA_DISABLE_DEPRECATION_WARNINGS
-	CacheBones(Context);
-	PRAGMA_ENABLE_DEPRECATION_WARNINGS
 }
 
 void FAnimNode_Base::Update_AnyThread(const FAnimationUpdateContext& Context)
 {
-	// Call legacy implementation for backwards compatibility
-	PRAGMA_DISABLE_DEPRECATION_WARNINGS
-	Update(Context);
-	PRAGMA_ENABLE_DEPRECATION_WARNINGS
 }
 
 void FAnimNode_Base::Evaluate_AnyThread(FPoseContext& Output)
 {
-	// Call legacy implementation for backwards compatibility
-	PRAGMA_DISABLE_DEPRECATION_WARNINGS
-	Evaluate(Output);
-	PRAGMA_ENABLE_DEPRECATION_WARNINGS
 }
 
 void FAnimNode_Base::EvaluateComponentSpace_AnyThread(FComponentSpacePoseContext& Output)
 {
-	// Call legacy implementation for backwards compatibility
-	PRAGMA_DISABLE_DEPRECATION_WARNINGS
-	EvaluateComponentSpace(Output);
-	PRAGMA_ENABLE_DEPRECATION_WARNINGS
 }
 
 bool FAnimNode_Base::IsLODEnabled(FAnimInstanceProxy* AnimInstanceProxy)
@@ -135,10 +109,6 @@ bool FAnimNode_Base::IsLODEnabled(FAnimInstanceProxy* AnimInstanceProxy)
 
 void FAnimNode_Base::OnInitializeAnimInstance(const FAnimInstanceProxy* InProxy, const UAnimInstance* InAnimInstance)
 {
-	// Call legacy implementation for backwards compatibility
-	PRAGMA_DISABLE_DEPRECATION_WARNINGS
-	RootInitialize(InProxy);
-	PRAGMA_ENABLE_DEPRECATION_WARNINGS
 }
 
 void FAnimNode_Base::ResetDynamics(ETeleportType InTeleportType)
diff --git a/Engine/Source/Runtime/Engine/Private/Animation/AnimationRuntime.cpp b/Engine/Source/Runtime/Engine/Private/Animation/AnimationRuntime.cpp
--- a/Engine/Source/Runtime/Engine/Private/Animation/AnimationRuntime.cpp
+++ b/Engine/Source/Runtime/Engine/Private/Animation/AnimationRuntime.cpp
@@ -2080,11 +2080,6 @@ void FAnimationRuntime::CreateMaskWeights(TArray<FPerBoneBlendWeight>& BoneBlend
 	}
 }
 
-void FAnimationRuntime::ConvertCSTransformToBoneSpace(USkeletalMeshComponent* SkelComp, FCSPose<FCompactPose>& MeshBases, FTransform& InOutCSBoneTM, FCompactPoseBoneIndex BoneIndex, EBoneControlSpace Space)
-{
-	ConvertCSTransformToBoneSpace(SkelComp->GetComponentTransform(), MeshBases, InOutCSBoneTM, BoneIndex, Space);
-}
-
 void FAnimationRuntime::ConvertCSTransformToBoneSpace(const FTransform& ComponentTransform, FCSPose<FCompactPose>& MeshBases, FTransform& InOutCSBoneTM, FCompactPoseBoneIndex BoneIndex, EBoneControlSpace Space)
 {
 	switch( Space )
@@ -2122,11 +2117,6 @@ void FAnimationRuntime::ConvertCSTransformToBoneSpace(const FTransform& Componen
 	}
 }
 
-void FAnimationRuntime::ConvertBoneSpaceTransformToCS(USkeletalMeshComponent* SkelComp, FCSPose<FCompactPose>& MeshBases, FTransform& InOutBoneSpaceTM, FCompactPoseBoneIndex BoneIndex, EBoneControlSpace Space)
-{
-	ConvertBoneSpaceTransformToCS(SkelComp->GetComponentTransform(), MeshBases, InOutBoneSpaceTM, BoneIndex, Space);
-}
-
 void FAnimationRuntime::ConvertBoneSpaceTransformToCS(const FTransform& ComponentTransform, FCSPose<FCompactPose>& MeshBases, FTransform& InOutBoneSpaceTM, FCompactPoseBoneIndex BoneIndex, EBoneControlSpace Space)
 {
 	switch( Space )
diff --git a/Engine/Source/Runtime/Engine/Private/AudioDevice.cpp b/Engine/Source/Runtime/Engine/Private/AudioDevice.cpp
--- a/Engine/Source/Runtime/Engine/Private/AudioDevice.cpp
+++ b/Engine/Source/Runtime/Engine/Private/AudioDevice.cpp
@@ -5984,38 +5984,6 @@ bool FAudioDevice::FCreateComponentParams::ShouldUseAttenuation() const
 	return true;
 }
 
-UAudioComponent* FAudioDevice::CreateComponent(USoundBase* Sound, UWorld* World, AActor* Actor, bool bPlay, bool bStopWhenOwnerDestroyed, const FVector* Location, USoundAttenuation* AttenuationSettings, USoundConcurrency* ConcurrencySettings)
-{
-	TUniquePtr<FCreateComponentParams> Params;
-
-	if (Actor)
-	{
-		Params = MakeUnique<FCreateComponentParams>(Actor);
-	}
-	else if (World)
-	{
-		Params = MakeUnique<FCreateComponentParams>(World);
-	}
-	else
-	{
-		Params = MakeUnique<FCreateComponentParams>(GEngine->GetMainAudioDeviceRaw());
-	}
-
-	Params->bPlay = bPlay;
-	Params->bStopWhenOwnerDestroyed = bStopWhenOwnerDestroyed;
-	Params->AttenuationSettings = AttenuationSettings;
-
-	if (ConcurrencySettings)
-	{
-		Params->ConcurrencySet.Add(ConcurrencySettings);
-	}
-	if (Location)
-	{
-		Params->SetLocation(*Location);
-	}
-	return CreateComponent(Sound, *Params);
-}
-
 UAudioComponent* FAudioDevice::CreateComponent(USoundBase* Sound, const FCreateComponentParams& Params)
 {
 	check(IsInGameThread());
diff --git a/Engine/Source/Runtime/Engine/Private/BatchedElements.cpp b/Engine/Source/Runtime/Engine/Private/BatchedElements.cpp
--- a/Engine/Source/Runtime/Engine/Private/BatchedElements.cpp
+++ b/Engine/Source/Runtime/Engine/Private/BatchedElements.cpp
@@ -768,25 +768,6 @@ FSceneView FBatchedElements::CreateProxySceneView(const FMatrix& ProjectionMatri
 	return FSceneView(ProxyViewInitOptions);
 }
 
-bool FBatchedElements::Draw(FRHICommandList& RHICmdList, const FMeshPassProcessorRenderState& DrawRenderState, ERHIFeatureLevel::Type FeatureLevel, bool bNeedToSwitchVerticalAxis, const FMatrix& Transform, uint32 ViewportSizeX, uint32 ViewportSizeY, bool bHitTesting, float Gamma, const FSceneView* View, EBlendModeFilter::Type Filter) const
-{
-	if ( View )
-	{
-		// Going to ignore these parameters in favor of just using the values directly from the scene view, so ensure that they're identical.
-		check(Transform == View->ViewMatrices.GetViewProjectionMatrix());
-		check(ViewportSizeX == View->UnscaledViewRect.Width());
-		check(ViewportSizeY == View->UnscaledViewRect.Height());
-
-		return Draw(RHICmdList, DrawRenderState, FeatureLevel, bNeedToSwitchVerticalAxis, *View, bHitTesting, Gamma, Filter);
-	}
-	else
-	{
-		FIntRect ViewRect = FIntRect(0, 0, ViewportSizeX, ViewportSizeY);
-
-		return Draw(RHICmdList, DrawRenderState, FeatureLevel, bNeedToSwitchVerticalAxis, CreateProxySceneView(Transform, ViewRect), bHitTesting, Gamma, Filter);
-	}
-}
-
 bool FBatchedElements::Draw(FRHICommandList& RHICmdList, const FMeshPassProcessorRenderState& DrawRenderState, ERHIFeatureLevel::Type FeatureLevel, bool bNeedToSwitchVerticalAxis, const FSceneView& View, bool bHitTesting, float Gamma /* = 1.0f */, EBlendModeFilter::Type Filter /* = EBlendModeFilter::All */) const
 {
 	const FRelativeViewMatrices RelativeMatrices = FRelativeViewMatrices::Create(View.ViewMatrices);
diff --git a/Engine/Source/Runtime/Engine/Private/BoneContainer.cpp b/Engine/Source/Runtime/Engine/Private/BoneContainer.cpp
--- a/Engine/Source/Runtime/Engine/Private/BoneContainer.cpp
+++ b/Engine/Source/Runtime/Engine/Private/BoneContainer.cpp
@@ -639,9 +639,4 @@ bool FBoneReference::Initialize(const USkeleton* Skeleton)
 bool FBoneReference::IsValidToEvaluate(const FBoneContainer& RequiredBones) const
 {
 	return (BoneIndex != INDEX_NONE && RequiredBones.Contains(BoneIndex));
-}
-
-bool FBoneReference::IsValid(const FBoneContainer& RequiredBones) const
-{
-	return IsValidToEvaluate(RequiredBones);
-}
+}
\ No newline at end of file
diff --git a/Engine/Source/Runtime/Engine/Private/EdGraph/EdGraphNode.cpp b/Engine/Source/Runtime/Engine/Private/EdGraph/EdGraphNode.cpp
--- a/Engine/Source/Runtime/Engine/Private/EdGraph/EdGraphNode.cpp
+++ b/Engine/Source/Runtime/Engine/Private/EdGraph/EdGraphNode.cpp
@@ -169,9 +169,8 @@ UEdGraphNode::UEdGraphNode(const FObjectInitializer& ObjectInitializer)
 	, AdvancedPinDisplay(ENodeAdvancedPins::NoPins)
 	, EnabledState(ENodeEnabledState::Enabled)
 	, bUserSetEnabledState(false)
-	, bAllowSplitPins_DEPRECATED(false)
-	, bIsNodeEnabled_DEPRECATED(true)
 #if WITH_EDITORONLY_DATA
+	, bIsNodeEnabled_DEPRECATED(true)
 	, bCanResizeNode(false)
 	, bUnrelated(false)
 	, bCommentBubblePinned(false)
@@ -547,20 +546,11 @@ FString UEdGraphNode::GetDocumentationExcerptName() const
 	return FString::Printf(TEXT("%s%s"), MyClass->GetPrefixCPP(), *MyClass->GetName());
 }
 
-PRAGMA_DISABLE_DEPRECATION_WARNINGS
 FSlateIcon UEdGraphNode::GetIconAndTint(FLinearColor& OutColor) const
 {
-	// @todo: Remove with GetPaletteIcon
-	FName DeprecatedName = GetPaletteIcon(OutColor);
-	if (!DeprecatedName.IsNone())
-	{
-		return FSlateIcon("EditorStyle", DeprecatedName);
-	}
-	
 	static const FSlateIcon Icon = FSlateIcon("EditorStyle", "GraphEditor.Default_16x");
 	return Icon;
 }
-PRAGMA_ENABLE_DEPRECATION_WARNINGS
 
 FString UEdGraphNode::GetDescriptiveCompiledName() const
 {
diff --git a/Engine/Source/Runtime/Engine/Private/GameState.cpp b/Engine/Source/Runtime/Engine/Private/GameState.cpp
--- a/Engine/Source/Runtime/Engine/Private/GameState.cpp
+++ b/Engine/Source/Runtime/Engine/Private/GameState.cpp
@@ -43,11 +43,6 @@ void AGameState::DefaultTimer()
 	GetWorldTimerManager().SetTimer(TimerHandle_DefaultTimer, this, &AGameState::DefaultTimer, GetWorldSettings()->GetEffectiveTimeDilation() / GetWorldSettings()->DemoPlayTimeDilation, true);
 }
 
-bool AGameState::ShouldShowGore() const
-{
-	return true;
-}
-
 void AGameState::PostInitializeComponents()
 {
 	Super::PostInitializeComponents();
diff --git a/Engine/Source/Runtime/Engine/Private/Interpolation.cpp b/Engine/Source/Runtime/Engine/Private/Interpolation.cpp
--- a/Engine/Source/Runtime/Engine/Private/Interpolation.cpp
+++ b/Engine/Source/Runtime/Engine/Private/Interpolation.cpp
@@ -889,17 +889,6 @@ void AMatineeActor::InitInterp()
 		// Cache whether or not we want to enable extreme content within this sequence
 		bShouldShowGore = true;
 
-PRAGMA_DISABLE_DEPRECATION_WARNINGS
-		if( GetWorld() != NULL && GetWorld()->GetWorldSettings() != NULL )
-		{
-			AGameState const* const GameState = GetWorld()->GetGameState<AGameState>();
-			if( GameState != NULL )
-			{
-				bShouldShowGore = GameState->ShouldShowGore();
-			}
-		}
-PRAGMA_ENABLE_DEPRECATION_WARNINGS
-
 		for(int32 GroupIndex=0; GroupIndex<MatineeData->InterpGroups.Num(); GroupIndex++)
 		{
 			UInterpGroup* Group = MatineeData->InterpGroups[GroupIndex];
diff --git a/Engine/Source/Runtime/Engine/Private/Particles/ParticleComponents.cpp b/Engine/Source/Runtime/Engine/Private/Particles/ParticleComponents.cpp
--- a/Engine/Source/Runtime/Engine/Private/Particles/ParticleComponents.cpp
+++ b/Engine/Source/Runtime/Engine/Private/Particles/ParticleComponents.cpp
@@ -8339,17 +8339,7 @@ void AEmitterCameraLensEffectBase::ActivateLensEffect()
 	check(World);
 	if( !IsNetMode(NM_DedicatedServer) )
 	{
-		UParticleSystem* PSToActuallySpawn = PS_CameraEffect;
-
-PRAGMA_DISABLE_DEPRECATION_WARNINGS
-		AGameState* GameState = World->GetGameState<AGameState>();
-		if(GameState && !GameState->ShouldShowGore() )
-		{
-			PSToActuallySpawn = PS_CameraEffectNonExtremeContent_DEPRECATED;
-		}
-PRAGMA_ENABLE_DEPRECATION_WARNINGS
-
-		if( PSToActuallySpawn != NULL )
+		if (PS_CameraEffect)
 		{
 			SetTemplate( PS_CameraEffect );
 		}
@@ -8372,11 +8362,6 @@ bool AEmitterCameraLensEffectBase::IsLooping() const
 		return true;
 	}
 
-	if ((PS_CameraEffectNonExtremeContent_DEPRECATED != nullptr) && PS_CameraEffectNonExtremeContent_DEPRECATED->IsLooping())
-	{
-		return true;
-	}
-
 	return false;
 }
 
diff --git a/Engine/Source/Runtime/Engine/Private/PlayerController.cpp b/Engine/Source/Runtime/Engine/Private/PlayerController.cpp
--- a/Engine/Source/Runtime/Engine/Private/PlayerController.cpp
+++ b/Engine/Source/Runtime/Engine/Private/PlayerController.cpp
@@ -3946,23 +3946,6 @@ void APlayerController::K2_ClientPlayForceFeedback(class UForceFeedbackEffect* F
 	ClientPlayForceFeedback(ForceFeedbackEffect, Params);
 }
 
-void APlayerController::ClientPlayForceFeedback(class UForceFeedbackEffect* ForceFeedbackEffect, bool bLooping, bool bIgnoreTimeDilation, FName Tag)
-{
-	FForceFeedbackParameters Params;
-	Params.Tag = Tag;
-	Params.bLooping = bLooping;
-	Params.bIgnoreTimeDilation = bIgnoreTimeDilation;
-	ClientPlayForceFeedback(ForceFeedbackEffect, Params);
-}
-
-void APlayerController::ClientPlayForceFeedback(class UForceFeedbackEffect* ForceFeedbackEffect, bool bLooping, FName Tag)
-{
-	FForceFeedbackParameters Params;
-	Params.Tag = Tag;
-	Params.bLooping = bLooping;
-	ClientPlayForceFeedback(ForceFeedbackEffect, Params);
-}
-
 void APlayerController::ClientStopForceFeedback_Implementation( UForceFeedbackEffect* ForceFeedbackEffect, FName Tag)
 {
 	if (ForceFeedbackEffect == NULL && Tag == NAME_None)
diff --git a/Engine/Source/Runtime/Engine/Public/Animation/AnimInstanceProxy.h b/Engine/Source/Runtime/Engine/Public/Animation/AnimInstanceProxy.h
--- a/Engine/Source/Runtime/Engine/Public/Animation/AnimInstanceProxy.h
+++ b/Engine/Source/Runtime/Engine/Public/Animation/AnimInstanceProxy.h
@@ -179,14 +179,6 @@ public:
 		return AnimClassInterface;
 	}
 
-	// Get the Blueprint Generated Class associated with this context, if there is one.
-	// Note: This can return NULL, so check the result.
-	UE_DEPRECATED(4.11, "GetAnimBlueprintClass() is deprecated, UAnimBlueprintGeneratedClass should not be directly used at runtime. Please use GetAnimClassInterface() instead.")
-	UAnimBlueprintGeneratedClass* GetAnimBlueprintClass() const
-	{
-		return Cast<UAnimBlueprintGeneratedClass>(IAnimClassInterface::GetActualAnimClass(AnimClassInterface));
-	}
-
 	/** Get the last DeltaSeconds passed into PreUpdate() */
 	float GetDeltaSeconds() const
 	{
diff --git a/Engine/Source/Runtime/Engine/Public/AnimationRuntime.h b/Engine/Source/Runtime/Engine/Public/AnimationRuntime.h
--- a/Engine/Source/Runtime/Engine/Public/AnimationRuntime.h
+++ b/Engine/Source/Runtime/Engine/Public/AnimationRuntime.h
@@ -551,9 +551,6 @@ public:
 	 */
 	static void ConvertCSTransformToBoneSpace(const FTransform& ComponentTransform, FCSPose<FCompactPose>& MeshBases, FTransform& InOutCSBoneTM, FCompactPoseBoneIndex BoneIndex, EBoneControlSpace Space);
 
-	UE_DEPRECATED(4.16, "Please use the ConvertCSTransformToBoneSpace with a transform as the first argument")
-	static void ConvertCSTransformToBoneSpace(USkeletalMeshComponent* SkelComp, FCSPose<FCompactPose>& MeshBases, FTransform& InOutCSBoneTM, FCompactPoseBoneIndex BoneIndex, EBoneControlSpace Space);
-
 	/** 
 	 * Convert a FTransform in a specified bone space to ComponentSpace.
 	 * @param	ComponentTransform	The transform of the component. Only used if Space == BCS_WorldSpace
@@ -564,9 +561,6 @@ public:
 	 */
 	static void ConvertBoneSpaceTransformToCS(const FTransform& ComponentTransform, FCSPose<FCompactPose>& MeshBases, FTransform& InOutBoneSpaceTM, FCompactPoseBoneIndex BoneIndex, EBoneControlSpace Space);
 
-	UE_DEPRECATED(4.16, "Please use the ConvertBoneSpaceTransformToCS with a transform as the first argument")
-	static void ConvertBoneSpaceTransformToCS(USkeletalMeshComponent* SkelComp, FCSPose<FCompactPose>& MeshBases, FTransform& InOutBoneSpaceTM, FCompactPoseBoneIndex BoneIndex, EBoneControlSpace Space);
-
 	// FA2Pose/FA2CSPose Interfaces for template functions
 	static FTransform GetSpaceTransform(FA2Pose& Pose, int32 Index);
 	static FTransform GetSpaceTransform(FA2CSPose& Pose, int32 Index);
diff --git a/Engine/Source/Runtime/Engine/Public/AudioDevice.h b/Engine/Source/Runtime/Engine/Public/AudioDevice.h
--- a/Engine/Source/Runtime/Engine/Public/AudioDevice.h
+++ b/Engine/Source/Runtime/Engine/Public/AudioDevice.h
@@ -735,20 +735,10 @@ public:
 
 	/**
 	 * Creates an audio component to handle playing a sound.
-	 * Plays a sound at the given location without creating an audio component.
 	 * @param   Sound				The USoundBase to play at the location.
-	 * @param   World				The world this sound is playing in.
-	 * @param   AActor				The optional actor with which to play the sound on.
-	 * @param   Play				Whether or not to automatically call play on the audio component after it is created.
-	 * @param	bStopWhenOwnerDestroyed Whether or not to automatically stop the audio component if its owner is destroyed.
-	 * @param	Location			The sound's location.
-	 * @param	AttenuationSettings	The sound's attenuation settings to use. Will default to the USoundBase's AttenuationSettings if not specified.
-	 * @param	USoundConcurrency	The sound's sound concurrency settings to use. Will use the USoundBase's USoundConcurrency if not specified.
+	 * @param   Params				The parameter block of properties to create the component based on
 	 * @return	The created audio component if the function successfully created one or a nullptr if not successful. Note: if audio is disabled or if there were no hardware audio devices available, this will return nullptr.
 	 */
-	UE_DEPRECATED(4.14, "Use CreateComponent that passes a parameters block instead")
-	static UAudioComponent* CreateComponent(USoundBase* Sound, UWorld* World, AActor* Actor = nullptr, bool bPlay = true, bool bStopWhenOwnerDestroyed = false, const FVector* Location = nullptr, USoundAttenuation* AttenuationSettings = nullptr, USoundConcurrency* ConcurrencySettings = nullptr);
-
 	static UAudioComponent* CreateComponent(USoundBase* Sound, const FCreateComponentParams& Params = FCreateComponentParams());
 
 	/**
diff --git a/Engine/Source/Runtime/Engine/Public/BatchedElements.h b/Engine/Source/Runtime/Engine/Public/BatchedElements.h
--- a/Engine/Source/Runtime/Engine/Public/BatchedElements.h
+++ b/Engine/Source/Runtime/Engine/Public/BatchedElements.h
@@ -231,20 +231,6 @@ public:
 		float OpacityMaskRefVal = .5f
 		);
 
-	/** 
-	 *Draws the batch
-	 *
-	 * @param Transform	The transform matrix for each viewport (View->ViewProjectionMatrix)
-	 * @param ViewportSizeX	The width of the viewport
-	 * @param ViewportSizeY The height of the viewport
-	 * @param bHitTesting	Whether or not we are hit testing
-	 * @param Gamma			Optional gamma override
-	 * @param View			Optional FSceneView for shaders that need access to view constants
-	 * @param DepthTexture	DepthTexture for manual depth testing with editor compositing in the pixel shader
-	 */
-	UE_DEPRECATED(4.14, "Deprecated. Use the FBatchedElements::Draw method that takes a non-optional FSceneView parameter instead")
-	bool Draw(FRHICommandList& RHICmdList, const FMeshPassProcessorRenderState& DrawRenderState, ERHIFeatureLevel::Type FeatureLevel, bool bNeedToSwitchVerticalAxis, const FMatrix& Transform, uint32 ViewportSizeX, uint32 ViewportSizeY, bool bHitTesting, float Gamma = 1.0f, const FSceneView* View = nullptr, EBlendModeFilter::Type Filter = EBlendModeFilter::All) const;
-	
 	/**
 	 * Draws the batch
 	 *
diff --git a/Engine/Source/Runtime/Engine/Public/BoneContainer.h b/Engine/Source/Runtime/Engine/Public/BoneContainer.h
--- a/Engine/Source/Runtime/Engine/Public/BoneContainer.h
+++ b/Engine/Source/Runtime/Engine/Public/BoneContainer.h
@@ -752,10 +752,7 @@ struct FBoneReference
 	// it triggers ensure in those functions
 	ENGINE_API bool Initialize(const USkeleton* Skeleton);
 
-	/** Deprecated functions */
-	UE_DEPRECATED(4.17, "Please use IsValidToEvaluate instead")
-	ENGINE_API bool IsValid(const FBoneContainer& RequiredBones) const;
-	
+
 	/** return true if it has valid set up */
 	bool HasValidSetup() const
 	{
diff --git a/Engine/Source/Runtime/Engine/Public/Slate/SceneViewport.h b/Engine/Source/Runtime/Engine/Public/Slate/SceneViewport.h
--- a/Engine/Source/Runtime/Engine/Public/Slate/SceneViewport.h
+++ b/Engine/Source/Runtime/Engine/Public/Slate/SceneViewport.h
@@ -178,12 +178,6 @@ public:
 	 */
 	virtual void ResizeFrame(uint32 NewSizeX,uint32 NewSizeY,EWindowMode::Type NewWindowMode) override;
 
-	UE_DEPRECATED(4.13, "The version of FSceneViewport::ResizeFrame that takes a position is deprecated (the position was never used). Please use the version that doesn't take a position instead.")
-	void ResizeFrame(uint32 NewSizeX, uint32 NewSizeY, EWindowMode::Type NewWindowMode, int32, int32)
-	{
-		ResizeFrame(NewSizeX, NewSizeY, NewWindowMode);
-	}
-
 	/**
 	 *	Sets the Viewport resize delegate.
 	 */
@@ -306,12 +300,6 @@ private:
 	 */
 	virtual void ResizeViewport( uint32 NewSizeX,uint32 NewSizeY,EWindowMode::Type NewWindowMode );
 
-	UE_DEPRECATED(4.13, "The version of FSceneViewport::ResizeViewport that takes a position is deprecated (the position was never used). Please use the version that doesn't take a position instead.")
-	void ResizeViewport(uint32 NewSizeX, uint32 NewSizeY, EWindowMode::Type NewWindowMode, int32, int32)
-	{
-		ResizeViewport(NewSizeX, NewSizeY, NewWindowMode);
-	}
-
 	/**
 	 * Called from slate when input is finished for this frame, and we should process any accumulated mouse data.
 	 */
diff --git a/Engine/Source/Runtime/Engine/Public/UnrealClient.h b/Engine/Source/Runtime/Engine/Public/UnrealClient.h
--- a/Engine/Source/Runtime/Engine/Public/UnrealClient.h
+++ b/Engine/Source/Runtime/Engine/Public/UnrealClient.h
@@ -132,12 +132,6 @@ public:
 
 	virtual class FViewport* GetViewport() = 0;
 	virtual void ResizeFrame(uint32 NewSizeX,uint32 NewSizeY,EWindowMode::Type NewWindowMode) = 0;
-
-	UE_DEPRECATED(4.13, "The version of FViewportFrame::ResizeFrame that takes a position is deprecated (the position was never used). Please use the version that doesn't take a position instead.")
-	void ResizeFrame(uint32 NewSizeX, uint32 NewSizeY, EWindowMode::Type NewWindowMode, int32, int32)
-	{
-		ResizeFrame(NewSizeX, NewSizeY, NewWindowMode);
-	}
 };
 
 /**
diff --git a/Engine/Source/Runtime/ImageWrapper/Public/IImageWrapper.h b/Engine/Source/Runtime/ImageWrapper/Public/IImageWrapper.h
--- a/Engine/Source/Runtime/ImageWrapper/Public/IImageWrapper.h
+++ b/Engine/Source/Runtime/ImageWrapper/Public/IImageWrapper.h
@@ -239,8 +239,3 @@ public:
 	/** Virtual destructor. */
 	virtual ~IImageWrapper() { }
 };
-
-
-/** Type definition for shared pointers to instances of IImageWrapper. */
-UE_DEPRECATED(4.16, "IImageWrapperPtr is deprecated. Please use 'TSharedPtr<IImageWrapper>' instead!")
-typedef TSharedPtr<IImageWrapper> IImageWrapperPtr;
diff --git a/Engine/Source/Runtime/LevelSequence/Public/LevelSequencePlayer.h b/Engine/Source/Runtime/LevelSequence/Public/LevelSequencePlayer.h
--- a/Engine/Source/Runtime/LevelSequence/Public/LevelSequencePlayer.h
+++ b/Engine/Source/Runtime/LevelSequence/Public/LevelSequencePlayer.h
@@ -22,10 +22,6 @@ class FLevelSequenceSpawnRegister;
 class FViewportClient;
 class UCameraComponent;
 
-struct UE_DEPRECATED(4.15, "Please use FMovieSceneSequencePlaybackSettings.") FLevelSequencePlaybackSettings
-	: public FMovieSceneSequencePlaybackSettings
-{};
-
 DECLARE_DYNAMIC_MULTICAST_DELEGATE_OneParam(FOnLevelSequencePlayerCameraCutEvent, UCameraComponent*, CameraComponent);
 
 USTRUCT(BlueprintType)
@@ -133,13 +129,6 @@ public:
 
 public:
 
-	/**
-	 * Access the level sequence this player is playing
-	 * @return the level sequence currently assigned to this player
-	 */
-	UE_DEPRECATED(4.15, "Please use GetSequence instead.")
-	ULevelSequence* GetLevelSequence() const { return Cast<ULevelSequence>(Sequence); }
-
 	// IMovieScenePlayer interface
 	virtual UObject* GetPlaybackContext() const override;
 	virtual TArray<UObject*> GetEventContexts() const override;
diff --git a/Engine/Source/Runtime/Messaging/Public/IMessageContext.h b/Engine/Source/Runtime/Messaging/Public/IMessageContext.h
--- a/Engine/Source/Runtime/Messaging/Public/IMessageContext.h
+++ b/Engine/Source/Runtime/Messaging/Public/IMessageContext.h
@@ -373,12 +373,3 @@ public:
 	/** Virtual destructor. */
 	virtual ~IMessageContext() { }
 };
-
-
-/** Type definition for shared pointers to instances of IMessageContext. */
-UE_DEPRECATED(4.16, "IMessageContextPtr is deprecated. Please use 'TSharedPtr<IMessageContext, ESPMode::ThreadSafe>' instead!")
-typedef TSharedPtr<IMessageContext, ESPMode::ThreadSafe> IMessageContextPtr;
-
-/** Type definition for shared references to instances of IMessageContext. */
-UE_DEPRECATED(4.16, "IMessageContextRef is deprecated. Please use 'TSharedRef<IMessageContext, ESPMode::ThreadSafe>' instead!")
-typedef TSharedRef<IMessageContext, ESPMode::ThreadSafe> IMessageContextRef;
diff --git a/Engine/Source/Runtime/Messaging/Public/IMessageReceiver.h b/Engine/Source/Runtime/Messaging/Public/IMessageReceiver.h
--- a/Engine/Source/Runtime/Messaging/Public/IMessageReceiver.h
+++ b/Engine/Source/Runtime/Messaging/Public/IMessageReceiver.h
@@ -90,8 +90,4 @@ public:
 
 	/** Virtual destructor. */
 	virtual ~IMessageReceiver() { }
-};
-
-
-UE_DEPRECATED(4.15, "IReceiveMessages has been renamed to IMessageReceiver")
-typedef IMessageReceiver IReceiveMessages;
+};
\ No newline at end of file
diff --git a/Engine/Source/Runtime/Messaging/Public/IMessageSender.h b/Engine/Source/Runtime/Messaging/Public/IMessageSender.h
--- a/Engine/Source/Runtime/Messaging/Public/IMessageSender.h
+++ b/Engine/Source/Runtime/Messaging/Public/IMessageSender.h
@@ -48,7 +48,3 @@ public:
 	/** Virtual destructor. */
 	virtual ~IMessageSender() { }
 };
-
-
-UE_DEPRECATED(4.15, "ISendMessages has been renamed to IMessageSender")
-typedef IMessageSender ISendMessages;
diff --git a/Engine/Source/Runtime/Messaging/Public/IMessageTransport.h b/Engine/Source/Runtime/Messaging/Public/IMessageTransport.h
--- a/Engine/Source/Runtime/Messaging/Public/IMessageTransport.h
+++ b/Engine/Source/Runtime/Messaging/Public/IMessageTransport.h
@@ -57,9 +57,4 @@ protected:
 
 	/** Virtual destructor. */
 	virtual ~IMessageTransport() { }
-};
-
-
-/** Type definition for shared references to instances of ITransportMessages. */
-UE_DEPRECATED(4.16, "IMessageTransportRef is deprecated. Please use 'TSharedRef<IMessageTransport, ESPMode::ThreadSafe>' instead!")
-typedef TSharedRef<IMessageTransport, ESPMode::ThreadSafe> IMessageTransportRef;
+};
\ No newline at end of file
diff --git a/Engine/Source/Runtime/MessagingCommon/Public/MessageEndpoint.h b/Engine/Source/Runtime/MessagingCommon/Public/MessageEndpoint.h
--- a/Engine/Source/Runtime/MessagingCommon/Public/MessageEndpoint.h
+++ b/Engine/Source/Runtime/MessagingCommon/Public/MessageEndpoint.h
@@ -1039,12 +1039,3 @@ private:
 	/** Sigfnifies that the handler array is being accessed and other threads should wait or skip */
 	FCriticalSection		HandlersCS;
 };
-
-
-/** Type definition for shared pointers to instances of FMessageEndpoint. */
-UE_DEPRECATED(4.16, "FMessageEndpointPtr is deprecated. Please use 'TSharedPtr<FMessageEndpoint, ESPMode::ThreadSafe>' instead!")
-typedef TSharedPtr<FMessageEndpoint, ESPMode::ThreadSafe> FMessageEndpointPtr;
-
-/** Type definition for shared references to instances of FMessageEndpoint. */
-UE_DEPRECATED(4.16, "FMessageEndpointRef is deprecated. Please use 'TSharedPtr<FMessageEndpoint, ESPMode::ThreadSafe>' instead!")
-typedef TSharedRef<FMessageEndpoint, ESPMode::ThreadSafe> FMessageEndpointRef;
diff --git a/Engine/Source/Runtime/Online/BuildPatchServices/Public/Interfaces/IBuildPatchServicesModule.h b/Engine/Source/Runtime/Online/BuildPatchServices/Public/Interfaces/IBuildPatchServicesModule.h
--- a/Engine/Source/Runtime/Online/BuildPatchServices/Public/Interfaces/IBuildPatchServicesModule.h
+++ b/Engine/Source/Runtime/Online/BuildPatchServices/Public/Interfaces/IBuildPatchServicesModule.h
@@ -227,16 +227,6 @@ public:
 	UE_DEPRECATED(4.21, "MakeManifestFromJSON(const FString& ManifestJSON) has been deprecated.  Please use MakeManifestFromData(const TArray<uint8>& ManifestData) instead.")
 	virtual IBuildManifestPtr MakeManifestFromJSON(const FString& ManifestJSON) = 0;
 
-	UE_DEPRECATED(4.16, "Please use EnumeratePatchData instead.")
-	virtual bool EnumerateManifestData(const FString& ManifestFilePath, const FString& OutputFile, bool bIncludeSizes)
-	{
-		BuildPatchServices::FPatchDataEnumerationConfiguration Configuration;
-		Configuration.InputFile = ManifestFilePath;
-		Configuration.OutputFile = OutputFile;
-		Configuration.bIncludeSizes = bIncludeSizes;
-		return EnumeratePatchData(Configuration);
-	}
-
 	UE_DEPRECATED(4.21, "Please use ChunkBuildDirectory instead.")
 	virtual bool GenerateChunksManifestFromDirectory(const BuildPatchServices::FGenerationConfiguration& Configuration)
 	{
diff --git a/Engine/Source/Runtime/Projects/Private/PluginManager.cpp b/Engine/Source/Runtime/Projects/Private/PluginManager.cpp
--- a/Engine/Source/Runtime/Projects/Private/PluginManager.cpp
+++ b/Engine/Source/Runtime/Projects/Private/PluginManager.cpp
@@ -1887,28 +1887,6 @@ TArray<TSharedRef<IPlugin>> FPluginManager::GetDiscoveredPlugins()
 	return Plugins;
 }
 
-TArray< FPluginStatus > FPluginManager::QueryStatusForAllPlugins() const
-{
-	TArray< FPluginStatus > PluginStatuses;
-	PluginStatuses.Reserve(AllPlugins.Num());
-
-	for( const TPair<FString, TSharedRef<FPlugin>>& PluginPair : AllPlugins )
-	{
-		const TSharedRef< FPlugin >& Plugin = PluginPair.Value;
-		
-		FPluginStatus PluginStatus;
-		PluginStatus.Name = Plugin->Name;
-		PluginStatus.PluginDirectory = FPaths::GetPath(Plugin->FileName);
-		PluginStatus.bIsEnabled = Plugin->bEnabled;
-		PluginStatus.Descriptor = Plugin->Descriptor;
-		PluginStatus.LoadedFrom = Plugin->GetLoadedFrom();
-
-		PluginStatuses.Add( PluginStatus );
-	}
-
-	return PluginStatuses;
-}
-
 bool FPluginManager::AddPluginSearchPath(const FString& ExtraDiscoveryPath, bool bRefresh)
 {
 	bool bAlreadyExists = false;
diff --git a/Engine/Source/Runtime/Projects/Private/PluginManager.h b/Engine/Source/Runtime/Projects/Private/PluginManager.h
--- a/Engine/Source/Runtime/Projects/Private/PluginManager.h
+++ b/Engine/Source/Runtime/Projects/Private/PluginManager.h
@@ -137,7 +137,6 @@ public:
 	virtual TArray<TSharedRef<IPlugin>> GetEnabledPluginsWithVerse() const override;
 	virtual TArray<TSharedRef<IPlugin>> GetEnabledPluginsWithContentOrVerse() const override;
 	virtual TArray<TSharedRef<IPlugin>> GetDiscoveredPlugins() override;
-	virtual TArray< FPluginStatus > QueryStatusForAllPlugins() const override;
 	virtual bool AddPluginSearchPath(const FString& ExtraDiscoveryPath, bool bRefresh = true) override;
 	const TSet<FString>& GetAdditionalPluginSearchPaths() const override;
 	virtual TArray<TSharedRef<IPlugin>> GetPluginsWithPakFile() const override;
diff --git a/Engine/Source/Runtime/Projects/Public/Interfaces/IPluginManager.h b/Engine/Source/Runtime/Projects/Public/Interfaces/IPluginManager.h
--- a/Engine/Source/Runtime/Projects/Public/Interfaces/IPluginManager.h
+++ b/Engine/Source/Runtime/Projects/Public/Interfaces/IPluginManager.h
@@ -323,14 +323,6 @@ public:
 	 */
 	virtual TArray<TSharedRef<IPlugin>> GetDiscoveredPlugins() = 0;
 
-	/**
-	 * Gets status about all currently known plug-ins.
-	 *
-	 * @return	 Array of plug-in status objects.
-	 */
-	UE_DEPRECATED(4.18, "QueryStatusForAllPlugins() has been deprecated. Please use GetDiscoveredPlugins() instead.")
-	virtual TArray<FPluginStatus> QueryStatusForAllPlugins() const = 0;
-
 	/**
 	 * Stores the specified path, utilizing it in future search passes when 
 	 * searching for available plugins. Optionally refreshes the manager after 
diff --git a/Engine/Source/Runtime/Slate/Private/Framework/Text/TextLayout.cpp b/Engine/Source/Runtime/Slate/Private/Framework/Text/TextLayout.cpp
--- a/Engine/Source/Runtime/Slate/Private/Framework/Text/TextLayout.cpp
+++ b/Engine/Source/Runtime/Slate/Private/Framework/Text/TextLayout.cpp
@@ -2097,11 +2097,6 @@ bool FTextLayout::RemoveLine(int32 LineIndex)
 	return true;
 }
 
-void FTextLayout::AddLine( const TSharedRef< FString >& Text, const TArray< TSharedRef< IRun > >& Runs )
-{
-	AddLine(FNewLineData(Text, Runs));
-}
-
 void FTextLayout::AddLine( const FNewLineData& NewLine )
 {
 	{
diff --git a/Engine/Source/Runtime/Slate/Public/Framework/Text/TextLayout.h b/Engine/Source/Runtime/Slate/Public/Framework/Text/TextLayout.h
--- a/Engine/Source/Runtime/Slate/Public/Framework/Text/TextLayout.h
+++ b/Engine/Source/Runtime/Slate/Public/Framework/Text/TextLayout.h
@@ -438,9 +438,6 @@ public:
 		TArray<TSharedRef<IRun>> Runs;
 	};
 
-	UE_DEPRECATED(4.11, "Please use the version of AddLine that takes an FNewLineData parameter.")
-	void AddLine( const TSharedRef< FString >& Text, const TArray< TSharedRef< IRun > >& Runs );
-
 	void AddLine( const FNewLineData& NewLine );
 
 	void AddLines( const TArray<FNewLineData>& NewLines );
diff --git a/Engine/Source/Runtime/Slate/Public/Widgets/SViewport.h b/Engine/Source/Runtime/Slate/Public/Widgets/SViewport.h
--- a/Engine/Source/Runtime/Slate/Public/Widgets/SViewport.h
+++ b/Engine/Source/Runtime/Slate/Public/Widgets/SViewport.h
@@ -198,9 +198,6 @@ public:
 	 */
 	void SetActive(bool bActive);
 
-	UE_DEPRECATED(4.11, "SetWidgetToFocusOnActivate is no longer needed, remove this call.")
-	void SetWidgetToFocusOnActivate(TSharedPtr<SWidget> WidgetToFocus) { }
-
 public:
 
 	// SWidget interface
diff --git a/Engine/Source/Runtime/SlateCore/Private/Widgets/SWidget.cpp b/Engine/Source/Runtime/SlateCore/Private/Widgets/SWidget.cpp
--- a/Engine/Source/Runtime/SlateCore/Private/Widgets/SWidget.cpp
+++ b/Engine/Source/Runtime/SlateCore/Private/Widgets/SWidget.cpp
@@ -396,15 +396,8 @@ void SWidget::OnFocusLost(const FFocusEvent& InFocusEvent)
 {
 }
 
-void SWidget::OnFocusChanging(const FWeakWidgetPath& PreviousFocusPath, const FWidgetPath& NewWidgetPath)
-{
-}
-
 void SWidget::OnFocusChanging(const FWeakWidgetPath& PreviousFocusPath, const FWidgetPath& NewWidgetPath, const FFocusEvent& InFocusEvent)
 {
-	PRAGMA_DISABLE_DEPRECATION_WARNINGS
-	OnFocusChanging(PreviousFocusPath, NewWidgetPath);
-	PRAGMA_ENABLE_DEPRECATION_WARNINGS
 }
 
 FReply SWidget::OnKeyChar( const FGeometry& MyGeometry, const FCharacterEvent& InCharacterEvent )
diff --git a/Engine/Source/Runtime/SlateCore/Public/Input/Events.h b/Engine/Source/Runtime/SlateCore/Public/Input/Events.h
--- a/Engine/Source/Runtime/SlateCore/Public/Input/Events.h
+++ b/Engine/Source/Runtime/SlateCore/Public/Input/Events.h
@@ -40,9 +40,6 @@ enum class EFocusCause : uint8
 	WindowActivate,
 };
 
-UE_DEPRECATED(4.18, "EKeyboardFocusCause is deprecated and was renamed to EFocusCause. Please use that type instead.")
-typedef EFocusCause EKeyboardFocusCause;
-
 /**
  * FFocusEvent is used when notifying widgets about keyboard focus changes
  * It is passed to event handlers dealing with keyboard focus
diff --git a/Engine/Source/Runtime/SlateCore/Public/Layout/Geometry.h b/Engine/Source/Runtime/SlateCore/Public/Layout/Geometry.h
--- a/Engine/Source/Runtime/SlateCore/Public/Layout/Geometry.h
+++ b/Engine/Source/Runtime/SlateCore/Public/Layout/Geometry.h
@@ -438,23 +438,6 @@ public:
 		return AbsoluteToLocal(AbsoluteCoordinateRounded);
 	}
 	
-	/**
-	 * !!! DEPRECATED !!! This legacy function does not account for render transforms.
-	 * 
-	 * Returns a clipping rectangle corresponding to the allocated geometry's absolute position and size.
-	 * Note that the clipping rectangle starts 1 pixel above and left of the geometry because clipping is not
-	 * inclusive on the lower bound.
-	 * 
-	 * Absolute coordinates could be either desktop or window space depending on what space the root of the widget hierarchy is in.
-	 *
-	 * @return  Allotted geometry rectangle in absolute coordinates.
-	 */
-	UE_DEPRECATED(4.17, "This no longer represents any sort of clipping rect.  Please use GetLayoutBoundingRect() to get the layout rect of this geometry.")
-	FORCEINLINE_DEBUGGABLE FSlateRect GetClippingRect() const
-	{
-		return TransformRect(GetAccumulatedLayoutTransform(), FSlateRect(FVector2D(0.0f, 0.0f), Size));
-	}
-
 	FORCEINLINE_DEBUGGABLE FSlateRect GetLayoutBoundingRect() const
 	{
 		return GetLayoutBoundingRect(FSlateRect(FVector2D(0.0f, 0.0f), Size));
diff --git a/Engine/Source/Runtime/SlateCore/Public/Widgets/SWidget.h b/Engine/Source/Runtime/SlateCore/Public/Widgets/SWidget.h
--- a/Engine/Source/Runtime/SlateCore/Public/Widgets/SWidget.h
+++ b/Engine/Source/Runtime/SlateCore/Public/Widgets/SWidget.h
@@ -345,10 +345,6 @@ public:
 	 */
 	virtual void OnFocusLost(const FFocusEvent& InFocusEvent);
 
-	/** Called whenever a focus path is changing on all the widgets within the old and new focus paths */
-	UE_DEPRECATED(4.13, "Please use the newer version of OnFocusChanging that takes a FocusEvent")
-	virtual void OnFocusChanging(const FWeakWidgetPath& PreviousFocusPath, const FWidgetPath& NewWidgetPath);
-
 	/** Called whenever a focus path is changing on all the widgets within the old and new focus paths */
 	virtual void OnFocusChanging(const FWeakWidgetPath& PreviousFocusPath, const FWidgetPath& NewWidgetPath, const FFocusEvent& InFocusEvent);
 
diff --git a/Engine/Source/Runtime/SlateCore/Public/Widgets/SWindow.h b/Engine/Source/Runtime/SlateCore/Public/Widgets/SWindow.h
--- a/Engine/Source/Runtime/SlateCore/Public/Widgets/SWindow.h
+++ b/Engine/Source/Runtime/SlateCore/Public/Widgets/SWindow.h
@@ -325,14 +325,6 @@ public:
 		/** Should this window be focused immediately after it is shown? */
 		SLATE_ARGUMENT( bool, FocusWhenFirstShown )
 
-		UE_DEPRECATED(4.16, "ActivateWhenFirstShown(bool) is deprecated. Please use ActivationPolicy(EWindowActivationPolicy) instead")
-		FArguments& ActivateWhenFirstShown(bool bActivateWhenFirstShown)
-		{
-			// Previously ActivateWhenFirstShown was being used as always activating, so we use Always here to ensure same behavior.
-			_ActivationPolicy = bActivateWhenFirstShown ? EWindowActivationPolicy::Always : EWindowActivationPolicy::Never;
-			return Me();
-		}
-
 		/** When should this window be activated upon being shown? */
 		SLATE_ARGUMENT( EWindowActivationPolicy, ActivationPolicy )
 
@@ -798,12 +790,6 @@ public:
 		return WidgetFocusedOnDeactivate;
 	}
 
-	UE_DEPRECATED(4.16, "ActivateWhenFirstShown() is deprecated. Please use ActivationPolicy() instead.")
-	bool ActivateWhenFirstShown() const
-	{
-		return ActivationPolicy() != EWindowActivationPolicy::Never;
-	}
-
 	/** @return the window activation policy used when showing the window */
 	EWindowActivationPolicy ActivationPolicy() const;
 
diff --git a/Engine/Source/Runtime/UMG/Private/Components/Widget.cpp b/Engine/Source/Runtime/UMG/Private/Components/Widget.cpp
--- a/Engine/Source/Runtime/UMG/Private/Components/Widget.cpp
+++ b/Engine/Source/Runtime/UMG/Private/Components/Widget.cpp
@@ -1078,11 +1078,6 @@ const FText UWidget::GetPaletteCategory()
 	return LOCTEXT("Uncategorized", "Uncategorized");
 }
 
-const FSlateBrush* UWidget::GetEditorIcon()
-{
-	return nullptr;
-}
-
 EVisibility UWidget::GetVisibilityInDesigner() const
 {
 	return bHiddenInDesigner ? EVisibility::Collapsed : EVisibility::Visible;
diff --git a/Engine/Source/Runtime/UMG/Public/Components/Widget.h b/Engine/Source/Runtime/UMG/Public/Components/Widget.h
--- a/Engine/Source/Runtime/UMG/Public/Components/Widget.h
+++ b/Engine/Source/Runtime/UMG/Public/Components/Widget.h
@@ -943,10 +943,6 @@ public:
 	 */
 	virtual void OnCreationFromPalette() { }
 
-	/** Gets the editor icon */
-	UE_DEPRECATED(4.12, "GetEditorIcon is deprecated. Please define widget icons in your style set in the form ClassIcon.MyWidget, and register your style through FClassIconFinder::(Un)RegisterIconSource")
-	virtual const FSlateBrush* GetEditorIcon();
-
 	/** Allows general fixups and connections only used at editor time. */
 	virtual void ConnectEditorData() { }
 
@@ -1005,16 +1001,6 @@ protected:
 	/** Function called after the underlying SWidget is constructed. */
 	virtual void OnWidgetRebuilt();
 	
-#if WITH_EDITOR
-	/** Utility method for building a design time wrapper widget. */
-	UE_DEPRECATED(4.17, "Don't call this function in RebuildWidget any more.  Override RebuildDesignWidget, and build the wrapper there; widgets that derive from Panel already do this.  If you need to recreate the dashed outline you can use CreateDesignerOutline inside RebuildDesignWidget.")
-	TSharedRef<SWidget> BuildDesignTimeWidget(TSharedRef<SWidget> WrapWidget) { return CreateDesignerOutline(WrapWidget); }
-#else
-	/** Just returns the incoming widget in non-editor builds. */
-	UE_DEPRECATED(4.17, "Don't call this function in RebuildWidget any more.  Override RebuildDesignWidget, and build the wrapper there; widgets that derive from Panel already do this.  If you need to recreate the dashed outline you can use CreateDesignerOutline inside RebuildDesignWidget.")
-	FORCEINLINE TSharedRef<SWidget> BuildDesignTimeWidget(TSharedRef<SWidget> WrapWidget) { return WrapWidget; }
-#endif
-
 #if WITH_EDITOR
 	virtual TSharedRef<SWidget> RebuildDesignWidget(TSharedRef<SWidget> Content);
 
diff --git a/Engine/Source/Runtime/WebBrowser/Private/WebBrowserSingleton.cpp b/Engine/Source/Runtime/WebBrowser/Private/WebBrowserSingleton.cpp
--- a/Engine/Source/Runtime/WebBrowser/Private/WebBrowserSingleton.cpp
+++ b/Engine/Source/Runtime/WebBrowser/Private/WebBrowserSingleton.cpp
@@ -154,7 +154,6 @@ FString FWebBrowserSingleton::ApplicationCacheDir() const
 }
 
 
-PRAGMA_DISABLE_DEPRECATION_WARNINGS
 class FWebBrowserWindowFactory
 	: public IWebBrowserWindowFactory
 {
@@ -181,17 +180,18 @@ public:
 		bool ShowErrorMessage = true,
 		FColor BackgroundColor = FColor(255, 255, 255, 255)) override
 	{
-		return IWebBrowserModule::Get().GetSingleton()->CreateBrowserWindow(
-			OSWindowHandle,
-			InitialURL,
-			bUseTransparency,
-			bThumbMouseButtonNavigation,
-			ContentsToLoad,
-			ShowErrorMessage,
-			BackgroundColor);
+		FCreateBrowserWindowSettings Settings;
+		Settings.OSWindowHandle = OSWindowHandle;
+		Settings.InitialURL = MoveTemp(InitialURL);
+		Settings.bUseTransparency = bUseTransparency;
+		Settings.bThumbMouseButtonNavigation = bThumbMouseButtonNavigation;
+		Settings.ContentsToLoad = MoveTemp(ContentsToLoad);
+		Settings.bShowErrorMessage = ShowErrorMessage;
+		Settings.BackgroundColor = BackgroundColor;
+
+		return IWebBrowserModule::Get().GetSingleton()->CreateBrowserWindow(Settings);
 	}
 };
-PRAGMA_ENABLE_DEPRECATION_WARNINGS
 
 class FNoWebBrowserWindowFactory
 	: public IWebBrowserWindowFactory
@@ -485,31 +485,6 @@ TSharedPtr<IWebBrowserWindow> FWebBrowserSingleton::CreateBrowserWindow(
 	return nullptr;
 }
 
-TSharedPtr<IWebBrowserWindow> FWebBrowserSingleton::CreateBrowserWindow(
-	void* OSWindowHandle,
-	FString InitialURL,
-	bool bUseTransparency,
-	bool bThumbMouseButtonNavigation,
-	TOptional<FString> ContentsToLoad,
-	bool ShowErrorMessage,
-	FColor BackgroundColor,
-	int BrowserFrameRate,
-	const TArray<FString>& AltRetryDomains)
-{
-	FCreateBrowserWindowSettings Settings;
-	Settings.OSWindowHandle = OSWindowHandle;
-	Settings.InitialURL = InitialURL;
-	Settings.bUseTransparency = false;// bUseTransparency;
-	Settings.bThumbMouseButtonNavigation = bThumbMouseButtonNavigation;
-	Settings.ContentsToLoad = ContentsToLoad;
-	Settings.bShowErrorMessage = ShowErrorMessage;
-	Settings.BackgroundColor = FColor::Black;// BackgroundColor;
-	Settings.BrowserFrameRate = BrowserFrameRate;
-	Settings.AltRetryDomains = AltRetryDomains;
-
-	return CreateBrowserWindow(Settings);
-}
-
 TSharedPtr<IWebBrowserWindow> FWebBrowserSingleton::CreateBrowserWindow(const FCreateBrowserWindowSettings& WindowSettings)
 {
 	bool bBrowserEnabled = true;
@@ -803,16 +778,6 @@ FString FWebBrowserSingleton::GetCurrentLocaleCode()
 	return LocaleCode;
 }
 
-PRAGMA_DISABLE_DEPRECATION_WARNINGS
-void FWebBrowserSingleton::DeleteBrowserCookies(FString URL, FString CookieName, TFunction<void(int)> Completed)
-{
-	if (DefaultCookieManager.IsValid())
-	{
-		DefaultCookieManager->DeleteCookies(URL, CookieName, Completed);
-	}
-}
-PRAGMA_ENABLE_DEPRECATION_WARNINGS
-
 TSharedPtr<IWebBrowserCookieManager> FWebBrowserSingleton::GetCookieManager(TOptional<FString> ContextId) const
 {
 	if (ContextId.IsSet())
diff --git a/Engine/Source/Runtime/WebBrowser/Private/WebBrowserSingleton.h b/Engine/Source/Runtime/WebBrowser/Private/WebBrowserSingleton.h
--- a/Engine/Source/Runtime/WebBrowser/Private/WebBrowserSingleton.h
+++ b/Engine/Source/Runtime/WebBrowser/Private/WebBrowserSingleton.h
@@ -78,25 +78,12 @@ public:
 		TSharedPtr<FCEFWebBrowserWindow>& BrowserWindowParent,
 		TSharedPtr<FWebBrowserWindowInfo>& BrowserWindowInfo) override;
 
-	TSharedPtr<IWebBrowserWindow> CreateBrowserWindow(
-		void* OSWindowHandle,
-		FString InitialURL,
-		bool bUseTransparency,
-		bool bThumbMouseButtonNavigation,
-		TOptional<FString> ContentsToLoad = TOptional<FString>(),
-		bool ShowErrorMessage = true,
-		FColor BackgroundColor = FColor(255, 255, 255, 255),
-		int BrowserFrameRate = 24,
-		const TArray<FString>& AltRetryDomains = TArray<FString>()) override;
-
 	TSharedPtr<IWebBrowserWindow> CreateBrowserWindow(const FCreateBrowserWindowSettings& Settings) override;
 
 #if	BUILD_EMBEDDED_APP
 	TSharedPtr<IWebBrowserWindow> CreateNativeBrowserProxy() override;
 #endif
 
-	virtual void DeleteBrowserCookies(FString URL = TEXT(""), FString CookieName = TEXT(""), TFunction<void(int)> Completed = nullptr) override;
-
 	virtual TSharedPtr<IWebBrowserCookieManager> GetCookieManager() const override
 	{
 		return DefaultCookieManager;
diff --git a/Engine/Source/Runtime/WebBrowser/Public/IWebBrowserSingleton.h b/Engine/Source/Runtime/WebBrowser/Public/IWebBrowserSingleton.h
--- a/Engine/Source/Runtime/WebBrowser/Public/IWebBrowserSingleton.h
+++ b/Engine/Source/Runtime/WebBrowser/Public/IWebBrowserSingleton.h
@@ -109,46 +109,15 @@ public:
 	/**
 	 * Create a new web browser window
 	 *
-	 * @param OSWindowHandle Handle of OS Window that the browser will be displayed in (can be null)
-	 * @param InitialURL URL that the browser should initially navigate to
-	 * @param bUseTransparency Whether to allow transparent rendering of pages
-	 * @param ContentsToLoad Optional string to load as a web page
-	 * @param ShowErrorMessage Whether to show an error message in case of loading errors.
-	 * @param BackgroundColor Opaque background color used before a document is loaded and when no document color is specified.
-	 * @param BrowserFrameRate The framerate of the browser in FPS
+	 * @param Settings Struct containing browser window settings
 	 * @return New Web Browser Window Interface (may be null if not supported)
 	 */
-	UE_DEPRECATED(4.11, "Please use the new overload that takes a settings struct.")
-	virtual TSharedPtr<IWebBrowserWindow> CreateBrowserWindow(
-		void* OSWindowHandle,
-		FString InitialURL,
-		bool bUseTransparency,
-		bool bThumbMouseButtonNavigation,
-		TOptional<FString> ContentsToLoad = TOptional<FString>(),
-		bool ShowErrorMessage = true,
-		FColor BackgroundColor = FColor(255, 255, 255, 255),
-		int BrowserFrameRate = 24,
-		const TArray<FString>& AltRetryDomains = TArray<FString>()) = 0;
-
 	virtual TSharedPtr<IWebBrowserWindow> CreateBrowserWindow(const FCreateBrowserWindowSettings& Settings) = 0;
 
 #if	BUILD_EMBEDDED_APP
 	virtual TSharedPtr<IWebBrowserWindow> CreateNativeBrowserProxy() = 0;
 #endif
 
-	/**
-	 * Delete all browser cookies.
-	 *
-	 * Removes all matching cookies. Leave both URL and CookieName blank to delete the entire cookie database.
-	 * The actual deletion will be scheduled on the browser IO thread, so the operation may not have completed when the function returns.
-	 *
-	 * @param URL The base URL to match when searching for cookies to remove. Use blank to match all URLs.
-	 * @param CookieName The name of the cookie to delete. If left unspecified, all cookies will be removed.
-	 * @param Completed A callback function that will be invoked asynchronously on the game thread when the deletion is complete passing in the number of deleted objects.
-	 */
-	UE_DEPRECATED(4.11, "Please use the CookieManager instead via GetCookieManager().")
-	virtual void DeleteBrowserCookies(FString URL = TEXT(""), FString CookieName = TEXT(""), TFunction<void(int)> Completed = nullptr) = 0;
-
 	virtual TSharedPtr<class IWebBrowserCookieManager> GetCookieManager() const = 0;
 
 	virtual TSharedPtr<class IWebBrowserCookieManager> GetCookieManager(TOptional<FString> ContextId) const = 0;
