WP HLOD - Add Landscape support

Changes to support WP landscape material capture
* Moved IsValidBaseMaterial() from MeshMergeUtilities to MaterialUtilities
* Moved CreateProxyMaterialInstance() from MeshMergeUtilities to MaterialUtilities
* Added FCreateTexture2DParameters::bVirtualTexture
* Added FMaterialUtilities::ExportLandscapeMaterial() which doesn't need a list of primitives to hide (will only render the landscape component primitives)
* Added VT primming to ensure the captured textures have the expected resolution

diff --git a/Engine/Source/Developer/MaterialUtilities/Private/MaterialUtilities.cpp b/Engine/Source/Developer/MaterialUtilities/Private/MaterialUtilities.cpp
--- a/Engine/Source/Developer/MaterialUtilities/Private/MaterialUtilities.cpp
+++ b/Engine/Source/Developer/MaterialUtilities/Private/MaterialUtilities.cpp
@@ -55,6 +55,12 @@ IMPLEMENT_MODULE(FMaterialUtilities, MaterialUtilities);
 
 DEFINE_LOG_CATEGORY_STATIC(LogMaterialUtilities, Log, All);
 
+static TAutoConsoleVariable<int32> CVarMaterialUtilitiesWarmupFrames(
+	TEXT("MaterialUtilities.WarmupFrames"),
+	10,
+	TEXT("Number of frames to render before each capture in order to warmup various rendering systems (VT/Nanite/etc)."));
+
+
 bool FMaterialUtilities::CurrentlyRendering = false;
 TArray<UTextureRenderTarget2D*> FMaterialUtilities::RenderTargetPool;
 
@@ -119,12 +125,27 @@ UMaterialInterface* FMaterialUtilities::CreateProxyMaterialAndTextures(UPackage*
 		// If the pixel data isn't constant create a texture for it 
 		if (ColorData.Num() > 1)
 		{
-			TextureCompressionSettings CompressionSettings = SpecialCompressionSettingProperties.Contains(Property) ? SpecialCompressionSettingProperties.FindChecked(Property) : TC_Default;
-			bool bSRGBEnabled = SRGBEnabledProperties.Contains(Property);
-			UTexture* Texture = FMaterialUtilities::CreateTexture(OuterPackage, TEXT("T_") + AssetName + TEXT("_") + TrimmedPropertyName, DataSize, ColorData, CompressionSettings, TEXTUREGROUP_HierarchicalLOD, RF_Public | RF_Standalone, bSRGBEnabled);
+			FMaterialParameterInfo ParameterInfo(*(TrimmedPropertyName + TEXT("Texture")));
+
+			FCreateTexture2DParameters CreateParams;
+			CreateParams.TextureGroup = TEXTUREGROUP_HierarchicalLOD;
+			CreateParams.CompressionSettings = SpecialCompressionSettingProperties.Contains(Property) ? SpecialCompressionSettingProperties.FindChecked(Property) : TC_Default;
+			CreateParams.bSRGB = SRGBEnabledProperties.Contains(Property);
+
+			// Make sure the texture is a VT if required by the material sampler
+			if (Material != nullptr)
+			{
+				UTexture* DefaultTexture = nullptr;
+				Material->GetTextureParameterValue(ParameterInfo, DefaultTexture);
+				if (DefaultTexture)
+				{
+					CreateParams.bVirtualTexture = DefaultTexture->VirtualTextureStreaming;
+				}
+			}
+
+			UTexture* Texture = FMaterialUtilities::CreateTexture(OuterPackage, TEXT("T_") + AssetName + TEXT("_") + TrimmedPropertyName, DataSize, ColorData, CreateParams, RF_Public | RF_Standalone);
 
 			// Set texture parameter value on instance material
-			FMaterialParameterInfo ParameterInfo(*(TrimmedPropertyName + TEXT("Texture")));
 			Material->SetTextureParameterValueEditorOnly(ParameterInfo, Texture);
 
 			FStaticSwitchParameter SwitchParameter;
@@ -730,16 +751,36 @@ private:
 	FGuid Id;
 };
 
+/**
+ * Render the scene to the provided canvas. Will potentially perform the render multiple times, depending on the value of
+ * the CVarMaterialUtilitiesWarmupFrames CVar. This is needed to ensure various rendering systems are primed properly before capturing
+ * the scene.
+ */
+static void PerformSceneRender(FCanvas& Canvas, FSceneViewFamily& ViewFamily, bool bWithWarmup)
+{
+	if (bWithWarmup)
+	{
+		for (int32 i = 0; i < CVarMaterialUtilitiesWarmupFrames.GetValueOnGameThread(); i++)
+		{
+			GetRendererModule().BeginRenderingViewFamily(&Canvas, &ViewFamily);
+		}
+	}
+
+	GetRendererModule().BeginRenderingViewFamily(&Canvas, &ViewFamily);
+}
+
 static void RenderSceneToTexture(
-		FSceneInterface* Scene,
-		const FName& VisualizationMode, 
-		const FVector& ViewOrigin,
-		const FMatrix& ViewRotationMatrix, 
-		const FMatrix& ProjectionMatrix,  
-		const TSet<FPrimitiveComponentId>& HiddenPrimitives, 
-		FIntPoint TargetSize,
-		float TargetGamma,
-		TArray<FColor>& OutSamples)
+	FSceneInterface* Scene,
+	const FName& VisualizationMode, 
+	const FVector& ViewOrigin,
+	const FMatrix& ViewRotationMatrix, 
+	const FMatrix& ProjectionMatrix,
+	const TSet<FPrimitiveComponentId>& ShowOnlyPrimitives,
+	const TSet<FPrimitiveComponentId>& HiddenPrimitives, 
+	FIntPoint TargetSize,
+	float TargetGamma,
+	bool bPerformWarmpup,
+	TArray<FColor>& OutSamples)
 {
 	auto RenderTargetTexture = NewObject<UTextureRenderTarget2D>();
 	check(RenderTargetTexture);
@@ -763,6 +804,7 @@ static void RenderSceneToTexture(
 	FSceneViewInitOptions ViewInitOptions;
 	ViewInitOptions.SetViewRectangle(FIntRect(0, 0, TargetSize.X, TargetSize.Y));
 	ViewInitOptions.ViewFamily = &ViewFamily;
+	ViewInitOptions.ShowOnlyPrimitives = ShowOnlyPrimitives;
 	ViewInitOptions.HiddenPrimitives = HiddenPrimitives;
 	ViewInitOptions.ViewOrigin = ViewOrigin;
 	ViewInitOptions.ViewRotationMatrix = ViewRotationMatrix;
@@ -785,19 +827,60 @@ static void RenderSceneToTexture(
 		Extension->SetupView(ViewFamily, *NewView);
 	}
 
-	GetRendererModule().BeginRenderingViewFamily(&Canvas, &ViewFamily);
+	PerformSceneRender(Canvas, ViewFamily, bPerformWarmpup);
 
 	// Copy the contents of the remote texture to system memory
 	OutSamples.SetNumUninitialized(TargetSize.X*TargetSize.Y);
 	FReadSurfaceDataFlags ReadSurfaceDataFlags;
 	ReadSurfaceDataFlags.SetLinearToGamma(false);
 	RenderTargetResource->ReadPixelsPtr(OutSamples.GetData(), ReadSurfaceDataFlags, FIntRect(0, 0, TargetSize.X, TargetSize.Y));
-	FlushRenderingCommands();
 					
 	RenderTargetTexture->RemoveFromRoot();
 	RenderTargetTexture = nullptr;
 }
 
+
+static void RenderSceneToTextures(
+	FSceneInterface* Scene,
+	const FVector& ViewOrigin,
+	const FMatrix& ViewRotationMatrix,
+	const FMatrix& ProjectionMatrix,
+	const TSet<FPrimitiveComponentId>& ShowOnlyPrimitives,
+	const TSet<FPrimitiveComponentId>& HiddenPrimitives,
+	FFlattenMaterial& OutFlattenMaterial)
+{
+	TMap<EFlattenMaterialProperties, TPair<FName, float>>	SupportedProperties; // Property -> VisualisationMode|Gamma
+	SupportedProperties.Add(EFlattenMaterialProperties::Diffuse) = TPair<FName, float>(FName("BaseColor"), true);	// BaseColor to gamma space
+	SupportedProperties.Add(EFlattenMaterialProperties::Normal) = TPair<FName, float>(FName("WorldNormal"), false);	// Dump normal texture in linear space
+	SupportedProperties.Add(EFlattenMaterialProperties::Metallic) = TPair<FName, float>(FName("Metallic"), false);	// Dump metallic texture in linear space
+	SupportedProperties.Add(EFlattenMaterialProperties::Roughness) = TPair<FName, float>(FName("Roughness"), true);	// Roughness material powers color by 2.2, transform it back to linear
+	SupportedProperties.Add(EFlattenMaterialProperties::Specular) = TPair<FName, float>(FName("Specular"), false);	// Dump specular texture in linear space
+
+	bool bPerformWarmpup = true;
+	for (int32 PropertyIndex = 0; PropertyIndex < (int32)EFlattenMaterialProperties::NumFlattenMaterialProperties; ++PropertyIndex)
+	{
+		const EFlattenMaterialProperties Property = (EFlattenMaterialProperties)PropertyIndex;
+		if (OutFlattenMaterial.ShouldGenerateDataForProperty(Property))
+		{
+			TPair<FName, float>* PropertyInfo = SupportedProperties.Find(Property);
+			if (PropertyInfo)
+			{
+				TArray<FColor>& Samples = OutFlattenMaterial.GetPropertySamples(Property);
+				const FIntPoint& Size = OutFlattenMaterial.GetPropertySize(Property);
+				const FName VisModeName = PropertyInfo->Key;
+				const float Gamma = PropertyInfo->Value ? 2.2f : 1.0f;
+
+				RenderSceneToTexture(Scene, VisModeName, ViewOrigin, ViewRotationMatrix, ProjectionMatrix, ShowOnlyPrimitives, HiddenPrimitives, Size, Gamma, bPerformWarmpup, Samples);
+				bPerformWarmpup = false; // Perform warmup only on first render
+			}
+			else
+			{
+				UE_LOG(LogMaterialUtilities, Error, TEXT("RenderSceneToTextures - Ignoring unsupported property"));
+			}
+		}
+	}
+}
+
 FIntPoint FMaterialUtilities::FindMaxTextureSize(UMaterialInterface* InMaterialInterface, FIntPoint MinimumSize)
 {
 	// static lod settings so that we only initialize them once
@@ -854,7 +937,7 @@ bool FMaterialUtilities::SupportsExport(EBlendMode InBlendMode, EMaterialPropert
 	return FExportMaterialProxy::WillFillData(InBlendMode, InMaterialProperty);
 }
 
-bool FMaterialUtilities::ExportLandscapeMaterial(ALandscapeProxy* InLandscape, const TSet<FPrimitiveComponentId>& HiddenPrimitives, FFlattenMaterial& OutFlattenMaterial)
+static bool ExportLandscapeMaterial(const ALandscapeProxy* InLandscape, const TSet<FPrimitiveComponentId>& ShowOnlyPrimitives, const TSet<FPrimitiveComponentId>& HiddenPrimitives, FFlattenMaterial& OutFlattenMaterial)
 {
 	check(InLandscape);
 
@@ -879,65 +962,31 @@ bool FMaterialUtilities::ExportLandscapeMaterial(ALandscapeProxy* InLandscape, c
 		ZOffset);
 
 	FSceneInterface* Scene = InLandscape->GetWorld()->Scene;
-						
-	// Render diffuse texture using BufferVisualizationMode=BaseColor
-	if (OutFlattenMaterial.ShouldGenerateDataForProperty(EFlattenMaterialProperties::Diffuse))
-	{
-		const FIntPoint& DiffuseSize = OutFlattenMaterial.GetPropertySize(EFlattenMaterialProperties::Diffuse);
-		static const FName BaseColorName("BaseColor");
-		const float BaseColorGamma = 2.2f; // BaseColor to gamma space
-		TArray<FColor>& DiffuseSamples = OutFlattenMaterial.GetPropertySamples(EFlattenMaterialProperties::Diffuse);
-		RenderSceneToTexture(Scene, BaseColorName, ViewOrigin, ViewRotationMatrix, ProjectionMatrix, HiddenPrimitives, 
-			DiffuseSize, BaseColorGamma, DiffuseSamples);
-	}
 
-	// Render normal map using BufferVisualizationMode=WorldNormal
-	// Final material should use world space instead of tangent space for normals
-	if (OutFlattenMaterial.ShouldGenerateDataForProperty(EFlattenMaterialProperties::Normal))
-	{
-		static const FName WorldNormalName("WorldNormal");
-		const float NormalColorGamma = 1.0f; // Dump normal texture in linear space
-		const FIntPoint& NormalSize = OutFlattenMaterial.GetPropertySize(EFlattenMaterialProperties::Normal);
-		TArray<FColor>& NormalSamples = OutFlattenMaterial.GetPropertySamples(EFlattenMaterialProperties::Normal);
-		RenderSceneToTexture(Scene, WorldNormalName, ViewOrigin, ViewRotationMatrix, ProjectionMatrix, HiddenPrimitives, 
-			NormalSize, NormalColorGamma, NormalSamples);
-	}
+	RenderSceneToTextures(Scene, ViewOrigin, ViewRotationMatrix, ProjectionMatrix, ShowOnlyPrimitives, HiddenPrimitives, OutFlattenMaterial);
 
-	// Render metallic map using BufferVisualizationMode=Metallic
-	if (OutFlattenMaterial.ShouldGenerateDataForProperty(EFlattenMaterialProperties::Metallic))
-	{
-		static const FName MetallicName("Metallic");
-		const float MetallicColorGamma = 1.0f; // Dump metallic texture in linear space
-		const FIntPoint& MetallicSize = OutFlattenMaterial.GetPropertySize(EFlattenMaterialProperties::Metallic);
-		TArray<FColor>& MetallicSamples = OutFlattenMaterial.GetPropertySamples(EFlattenMaterialProperties::Metallic);
-		RenderSceneToTexture(Scene, MetallicName, ViewOrigin, ViewRotationMatrix, ProjectionMatrix, HiddenPrimitives, 
-			MetallicSize, MetallicColorGamma, MetallicSamples);
-	}
+	OutFlattenMaterial.MaterialId = InLandscape->GetLandscapeGuid();
+	return true;
+}
 
-	// Render roughness map using BufferVisualizationMode=Roughness
-	if (OutFlattenMaterial.ShouldGenerateDataForProperty(EFlattenMaterialProperties::Roughness))
-	{
-		static const FName RoughnessName("Roughness");
-		const float RoughnessColorGamma = 2.2f; // Roughness material powers color by 2.2, transform it back to linear
-		const FIntPoint& RoughnessSize = OutFlattenMaterial.GetPropertySize(EFlattenMaterialProperties::Roughness);
-		TArray<FColor>& RoughnessSamples = OutFlattenMaterial.GetPropertySamples(EFlattenMaterialProperties::Roughness);
-		RenderSceneToTexture(Scene, RoughnessName, ViewOrigin, ViewRotationMatrix, ProjectionMatrix, HiddenPrimitives, 
-			RoughnessSize, RoughnessColorGamma, RoughnessSamples);
-	}
+bool FMaterialUtilities::ExportLandscapeMaterial(const ALandscapeProxy* InLandscape, FFlattenMaterial& OutFlattenMaterial)
+{
+	TSet<FPrimitiveComponentId> ShowOnlyPrimitives;
 
-	// Render specular map using BufferVisualizationMode=Specular
-	if (OutFlattenMaterial.ShouldGenerateDataForProperty(EFlattenMaterialProperties::Specular))
+	for (ULandscapeComponent* LandscapeComponent : InLandscape->LandscapeComponents)
 	{
-		static const FName SpecularName("Specular");
-		const float SpecularColorGamma = 1.0f; // Dump specular texture in linear space
-		const FIntPoint& SpecularSize = OutFlattenMaterial.GetPropertySize(EFlattenMaterialProperties::Specular);
-		TArray<FColor>& SpecularSamples = OutFlattenMaterial.GetPropertySamples(EFlattenMaterialProperties::Specular);
-		RenderSceneToTexture(Scene, SpecularName, ViewOrigin, ViewRotationMatrix, ProjectionMatrix, HiddenPrimitives, 
-			SpecularSize, SpecularColorGamma, SpecularSamples);
+		if (ensure(LandscapeComponent->SceneProxy))
+		{
+			ShowOnlyPrimitives.Add(LandscapeComponent->SceneProxy->GetPrimitiveComponentId());
+		}
 	}
-				
-	OutFlattenMaterial.MaterialId = InLandscape->GetLandscapeGuid();
-	return true;
+
+	return ::ExportLandscapeMaterial(InLandscape, ShowOnlyPrimitives, {}, OutFlattenMaterial);
+}
+
+bool FMaterialUtilities::ExportLandscapeMaterial(const ALandscapeProxy* InLandscape, const TSet<FPrimitiveComponentId>& HiddenPrimitives, FFlattenMaterial& OutFlattenMaterial)
+{
+	return ::ExportLandscapeMaterial(InLandscape, {}, HiddenPrimitives, OutFlattenMaterial);
 }
 
 UMaterial* FMaterialUtilities::CreateMaterial(const FFlattenMaterial& InFlattenMaterial, UPackage* InOuter, const FString& BaseName, EObjectFlags Flags, const struct FMaterialProxySettings& MaterialProxySettings, TArray<UObject*>& OutGeneratedAssets, const TextureGroup& InTextureGroup /*= TEXTUREGROUP_World*/)
@@ -1395,19 +1444,21 @@ UTexture2D* FMaterialUtilities::CreateTexture(UPackage* Outer, const FString& As
 	TexParams.bDeferCompression = true;
 	TexParams.bSRGB = bSRGB;
 	TexParams.SourceGuidHash = SourceGuidHash;
+	TexParams.TextureGroup = LODGroup;
+
+	return CreateTexture(Outer, AssetLongName, Size, Samples, TexParams, Flags);
+}
 
+UTexture2D* FMaterialUtilities::CreateTexture(UPackage* Outer, const FString& AssetLongName, FIntPoint Size, const TArray<FColor>& Samples, const FCreateTexture2DParameters& CreateParams, EObjectFlags Flags)
+{
 	if (Outer == nullptr)
 	{
-		Outer = CreatePackage( *AssetLongName);
+		Outer = CreatePackage(*AssetLongName);
 		Outer->FullyLoad();
 		Outer->Modify();
 	}
 
-	UTexture2D* Texture = FImageUtils::CreateTexture2D(Size.X, Size.Y, Samples, Outer, FPackageName::GetShortName(AssetLongName), Flags, TexParams);
-	Texture->LODGroup = LODGroup;
-	Texture->PostEditChange();
-			
-	return Texture;
+	return FImageUtils::CreateTexture2D(Size.X, Size.Y, Samples, Outer, FPackageName::GetShortName(AssetLongName), Flags, CreateParams);
 }
 
 bool FMaterialUtilities::ExportBaseColor(ULandscapeComponent* LandscapeComponent, int32 TextureSize, TArray<FColor>& OutSamples)
@@ -1438,19 +1489,15 @@ bool FMaterialUtilities::ExportBaseColor(ULandscapeComponent* LandscapeComponent
 	FSceneInterface* Scene = LandscapeProxy->GetWorld()->Scene;
 
 	// Hide all but the component
+	TSet<FPrimitiveComponentId> ShowOnlyPrimitives = { LandscapeComponent->SceneProxy->GetPrimitiveComponentId() };
 	TSet<FPrimitiveComponentId> HiddenPrimitives;
-	for (auto PrimitiveComponentId : Scene->GetScenePrimitiveComponentIds())
-	{
-		HiddenPrimitives.Add(PrimitiveComponentId);
-	}
-	HiddenPrimitives.Remove(LandscapeComponent->SceneProxy->GetPrimitiveComponentId());
 				
 	FIntPoint TargetSize(TextureSize, TextureSize);
 
 	// Render diffuse texture using BufferVisualizationMode=BaseColor
 	static const FName BaseColorName("BaseColor");
 	const float BaseColorGamma = 2.2f;
-	RenderSceneToTexture(Scene, BaseColorName, ViewOrigin, ViewRotationMatrix, ProjectionMatrix, HiddenPrimitives, TargetSize, BaseColorGamma, OutSamples);
+	RenderSceneToTexture(Scene, BaseColorName, ViewOrigin, ViewRotationMatrix, ProjectionMatrix, ShowOnlyPrimitives, HiddenPrimitives, TargetSize, BaseColorGamma, true, OutSamples);
 	return true;
 }
 
@@ -2511,7 +2558,6 @@ void FMaterialUtilities::ClearRenderTargetPool()
 
 void FMaterialUtilities::OptimizeSampleArray(TArray<FColor>& InSamples, FIntPoint& InSampleSize)
 {
-	// QQ LOOK AT 
 	if (InSamples.Num() > 1)
 	{
 		TArray<FColor> Colors;
@@ -2530,37 +2576,6 @@ void FMaterialUtilities::OptimizeSampleArray(TArray<FColor>& InSamples, FIntPoin
 			InSamples.Add(Colors[0]);
 			InSampleSize = FIntPoint(1, 1);
 		}
-
-		FColor ColourValue;
-		bool bValueFound = false;
-		for (FColor& Sample : InSamples)
-		{
-			if (Sample.A != 0)
-			{
-				ColourValue = Sample;
-				bValueFound = true;
-				break;
-			}
-		}
-
-		if (bValueFound)
-		{
-			bool bConstantValue = true;
-
-			for (FColor& Sample : InSamples)
-			{
-				if (Sample.A != 0 && ((Sample.DWColor() & FColor::Black.DWColor()) != (ColourValue.DWColor() & FColor::Black.DWColor())))
-				{
-					bConstantValue = false;
-					break;
-				}
-			}
-
-			if (bConstantValue)
-			{
-
-			}
-		}
 	}	
 }
 
diff --git a/Engine/Source/Developer/MaterialUtilities/Private/ProxyMaterialUtilities.cpp b/Engine/Source/Developer/MaterialUtilities/Private/ProxyMaterialUtilities.cpp
new file mode 100644
--- /dev/null
+++ b/Engine/Source/Developer/MaterialUtilities/Private/ProxyMaterialUtilities.cpp
@@ -0,0 +1,511 @@
+// Copyright Epic Games, Inc. All Rights Reserved.
+
+#pragma once
+
+#include "CoreMinimal.h"
+#include "Engine/MaterialMerging.h"
+#include "StaticParameterSet.h"
+#include "Materials/Material.h"
+#include "Engine/Texture2D.h"
+#include "MaterialUtilities.h"
+#include "Materials/MaterialInstanceConstant.h"
+#include "IMeshMergeUtilities.h"
+#include "MeshMergeModule.h"
+#include "Modules/ModuleManager.h"
+
+#if WITH_EDITOR
+#include "Widgets/Notifications/SNotificationList.h"
+#include "Framework/Notifications/NotificationManager.h"
+#endif // WITH_EDITOR
+
+#define LOCTEXT_NAMESPACE "ProxyMaterialUtilities"
+
+DEFINE_LOG_CATEGORY_STATIC(LogProxyMaterialUtilities, Log, All);
+
+// Find the parameter name to use with the provided material, for the given property
+static TArray<FString> GetPotentialParamNames(EFlattenMaterialProperties InProperty)
+{
+	switch(InProperty)
+	{
+		case EFlattenMaterialProperties::Diffuse: return { "BaseColor", "Diffuse" };
+		case EFlattenMaterialProperties::Normal: return { "Normal" };
+		case EFlattenMaterialProperties::Metallic: return { "Metallic" };
+		case EFlattenMaterialProperties::Roughness: return { "Roughness" };
+		case EFlattenMaterialProperties::Specular: return { "Specular" };
+		case EFlattenMaterialProperties::Opacity: return { "Opacity" };
+		case EFlattenMaterialProperties::OpacityMask: return { "OpacityMask" };
+		case EFlattenMaterialProperties::AmbientOcclusion: return { "AmbientOcclusion" };
+		case EFlattenMaterialProperties::Emissive: return { "EmissiveColor", "Emissive" };
+	}
+
+	return TArray<FString>();
+}
+
+static EMaterialParameterType GetConstantParamType(EFlattenMaterialProperties InProperty)
+{
+	EMaterialParameterType ParamType = EMaterialParameterType::None;
+
+	switch (InProperty)
+	{
+	case EFlattenMaterialProperties::Metallic:
+	case EFlattenMaterialProperties::Roughness:
+	case EFlattenMaterialProperties::Specular:
+	case EFlattenMaterialProperties::Opacity:
+	case EFlattenMaterialProperties::OpacityMask:
+	case EFlattenMaterialProperties::AmbientOcclusion:
+		ParamType = EMaterialParameterType::Scalar;
+		break;
+
+	case EFlattenMaterialProperties::Diffuse:
+	case EFlattenMaterialProperties::Emissive:
+		ParamType = EMaterialParameterType::Vector;
+		break;
+	}
+
+	return ParamType;
+}	
+
+// Find the parameter name to use with the provided material, for the given property
+static bool GetMatchingParamName(EFlattenMaterialProperties InProperty, const UMaterialInterface* InBaseMaterial, FString& OutParamName, TArray<FString>* OutMissingNames = nullptr)
+{
+	const TArray<FString> PotentialNames = GetPotentialParamNames(InProperty);
+
+	// Missing names, for error reporting
+	FString MissingNames;
+
+	for (const FString& PotentialName : PotentialNames)
+	{
+		const FName TextureName(PotentialName + TEXT("Texture"));
+		const FName ConstName(PotentialName + TEXT("Const"));
+		const FName UseTexture(TEXT("Use") + PotentialName);
+
+		UTexture* DefaultTexture = nullptr;
+		bool DefaultSwitchValue = false;
+		float DefaultScalarValue;
+		FLinearColor DefaultVectorValue;
+		FGuid ExpressionGuid;
+
+		if (!MissingNames.IsEmpty())
+		{
+			MissingNames += TEXT("|");
+		}
+
+		bool bHasRequiredParams = InBaseMaterial->GetTextureParameterValue(TextureName, DefaultTexture);
+
+		switch(GetConstantParamType(InProperty))
+		{
+		case EMaterialParameterType::Scalar:
+			MissingNames += UseTexture.ToString() + "+" + TextureName.ToString() + "+" + ConstName.ToString();
+			bHasRequiredParams &= InBaseMaterial->GetStaticSwitchParameterDefaultValue(UseTexture, DefaultSwitchValue, ExpressionGuid) && 
+									InBaseMaterial->GetScalarParameterDefaultValue(ConstName, DefaultScalarValue);
+			break;
+
+		case EMaterialParameterType::Vector:
+			MissingNames += UseTexture.ToString() + "+" + TextureName.ToString() + "+" + ConstName.ToString();
+			bHasRequiredParams &= InBaseMaterial->GetStaticSwitchParameterDefaultValue(UseTexture, DefaultSwitchValue, ExpressionGuid) &&
+									InBaseMaterial->GetVectorParameterDefaultValue(ConstName, DefaultVectorValue);
+			break;
+
+		case EMaterialParameterType::None:
+			MissingNames += UseTexture.ToString();
+			break;
+		}
+
+		if (bHasRequiredParams)
+		{
+			OutParamName = PotentialName;
+			return true;
+		}
+	}
+
+	if (OutMissingNames != nullptr)
+	{
+		OutMissingNames->Add(MissingNames);
+	}
+
+	OutParamName = "";
+	return false;
+}
+
+static FString GetMatchingParamName(EFlattenMaterialProperties InProperty, UMaterialInterface* InBaseMaterial)
+{
+	FString ParamName = "";
+	if (!GetMatchingParamName(InProperty, InBaseMaterial, ParamName))
+	{
+		UE_LOG(LogProxyMaterialUtilities, Fatal, TEXT("Invalid base material, should have been rejected by IsValidFlattenMaterial()"));
+	}
+	return ParamName;
+}
+
+// Validate that the provided material has all the required parameters needed to be considered a flattening material
+// If not, report what is missing
+bool FMaterialUtilities::IsValidFlattenMaterial(const UMaterialInterface* InBaseMaterial)
+{
+	if (InBaseMaterial != nullptr)
+	{
+		TArray<FGuid> ParameterIds;
+		TArray<FString> MissingParameters;
+
+
+		FString ParamName;
+		GetMatchingParamName(EFlattenMaterialProperties::Diffuse, InBaseMaterial, ParamName, &MissingParameters);
+		GetMatchingParamName(EFlattenMaterialProperties::Normal, InBaseMaterial, ParamName, &MissingParameters);
+		GetMatchingParamName(EFlattenMaterialProperties::Metallic, InBaseMaterial, ParamName, &MissingParameters);
+		GetMatchingParamName(EFlattenMaterialProperties::Roughness, InBaseMaterial, ParamName, &MissingParameters);
+		GetMatchingParamName(EFlattenMaterialProperties::Specular, InBaseMaterial, ParamName, &MissingParameters);
+		GetMatchingParamName(EFlattenMaterialProperties::Opacity, InBaseMaterial, ParamName, &MissingParameters);
+		GetMatchingParamName(EFlattenMaterialProperties::OpacityMask, InBaseMaterial, ParamName, &MissingParameters);
+		GetMatchingParamName(EFlattenMaterialProperties::AmbientOcclusion, InBaseMaterial, ParamName, &MissingParameters);
+		GetMatchingParamName(EFlattenMaterialProperties::Emissive, InBaseMaterial, ParamName, &MissingParameters);
+
+		auto NameCheckLambda = [&MissingParameters](const TArray<FMaterialParameterInfo>& InCheck, const TArray<FName>& InRequired)
+		{
+			for (const FName& Name : InRequired)
+			{
+				if (!InCheck.ContainsByPredicate([Name](const FMaterialParameterInfo& ParamInfo) { return (ParamInfo.Name == Name); }))
+				{
+					MissingParameters.Add(Name.ToString());
+				}
+			}
+		};
+
+		TArray<FMaterialParameterInfo> TextureParameterInfos;
+		TArray<FName> RequiredTextureNames = { TEXT("PackedTexture") };
+		InBaseMaterial->GetAllTextureParameterInfo(TextureParameterInfos, ParameterIds);
+		NameCheckLambda(TextureParameterInfos, RequiredTextureNames);
+
+		TArray<FMaterialParameterInfo> ScalarParameterInfos;
+		TArray<FName> RequiredScalarNames = { TEXT("EmissiveScale") };
+		InBaseMaterial->GetAllScalarParameterInfo(ScalarParameterInfos, ParameterIds);
+		NameCheckLambda(ScalarParameterInfos, RequiredScalarNames);
+
+		TArray<FMaterialParameterInfo> StaticSwitchParameterInfos;
+		TArray<FName> RequiredSwitchNames = { TEXT("PackMetallic"), TEXT("PackSpecular"), TEXT("PackRoughness") };
+		InBaseMaterial->GetAllStaticSwitchParameterInfo(StaticSwitchParameterInfos, ParameterIds);
+		NameCheckLambda(StaticSwitchParameterInfos, RequiredSwitchNames);
+
+		if (MissingParameters.Num() > 0)
+		{
+			FString MissingNamesString;
+			for (const FString& Name : MissingParameters)
+			{
+				if (!MissingNamesString.IsEmpty())
+				{
+					MissingNamesString += ", ";
+					MissingNamesString += Name;
+				}
+				else
+				{
+					MissingNamesString += Name;
+				}
+			}
+
+			UE_LOG(LogProxyMaterialUtilities, Warning, TEXT("Material %s is missing required Material Parameters %s"), *InBaseMaterial->GetName(), *MissingNamesString);
+
+			return false;
+		}
+		else
+		{
+			return true;
+		}
+	}
+
+	return false;
+}
+
+static const bool CalculatePackedTextureData(const FFlattenMaterial& InMaterial, bool& bOutPackMetallic, bool& bOutPackSpecular, bool& bOutPackRoughness, int32& OutNumSamples, FIntPoint& OutSize)
+{
+	// Whether or not a material property is baked down
+	const bool bHasMetallic = InMaterial.DoesPropertyContainData(EFlattenMaterialProperties::Metallic) && !InMaterial.IsPropertyConstant(EFlattenMaterialProperties::Metallic);
+	const bool bHasRoughness = InMaterial.DoesPropertyContainData(EFlattenMaterialProperties::Roughness) && !InMaterial.IsPropertyConstant(EFlattenMaterialProperties::Roughness);
+	const bool bHasSpecular = InMaterial.DoesPropertyContainData(EFlattenMaterialProperties::Specular) && !InMaterial.IsPropertyConstant(EFlattenMaterialProperties::Specular);
+
+	// Check for same texture sizes
+	bool bSameTextureSize = false;
+
+	// Determine whether or not the properties sizes match
+
+
+	const FIntPoint MetallicSize = InMaterial.GetPropertySize(EFlattenMaterialProperties::Metallic);
+	const FIntPoint SpecularSize = InMaterial.GetPropertySize(EFlattenMaterialProperties::Specular);
+	const FIntPoint RoughnessSize = InMaterial.GetPropertySize(EFlattenMaterialProperties::Roughness);
+
+	bSameTextureSize = (MetallicSize == RoughnessSize) || (MetallicSize == SpecularSize);
+	if (bSameTextureSize)
+	{
+		OutSize = MetallicSize;
+		OutNumSamples = InMaterial.GetPropertySamples(EFlattenMaterialProperties::Metallic).Num();
+	}
+	else
+	{
+		bSameTextureSize = (RoughnessSize == SpecularSize);
+		if (bSameTextureSize)
+		{
+			OutSize = RoughnessSize;
+			OutNumSamples = InMaterial.GetPropertySamples(EFlattenMaterialProperties::Roughness).Num();
+		}
+	}
+
+	// Now that we know if the data matches determine whether or not we should pack the properties
+	int32 NumPacked = 0;
+	if (OutNumSamples != 0)
+	{
+		bOutPackMetallic = bHasMetallic ? (OutNumSamples == InMaterial.GetPropertySamples(EFlattenMaterialProperties::Metallic).Num()) : false;
+		NumPacked += (bOutPackMetallic) ? 1 : 0;
+		bOutPackRoughness = bHasRoughness ? (OutNumSamples == InMaterial.GetPropertySamples(EFlattenMaterialProperties::Roughness).Num()) : false;
+		NumPacked += (bOutPackRoughness) ? 1 : 0;
+		bOutPackSpecular = bHasSpecular ? (OutNumSamples == InMaterial.GetPropertySamples(EFlattenMaterialProperties::Specular).Num()) : false;
+		NumPacked += (bOutPackSpecular) ? 1 : 0;
+	}
+	else
+	{
+		bOutPackMetallic = bOutPackRoughness = bOutPackSpecular = false;
+	}
+
+	// Need atleast two properties to pack
+	return NumPacked >= 2;
+}
+
+UMaterialInstanceConstant* FMaterialUtilities::CreateFlattenMaterialInstance(UPackage* InOuter, const FMaterialProxySettings& InMaterialProxySettings, UMaterialInterface* InBaseMaterial, const FFlattenMaterial& FlattenMaterial, const FString& AssetBasePath, const FString& AssetBaseName, TArray<UObject*>& OutAssetsToSync, FMaterialUpdateContext* MaterialUpdateContext)
+{
+	UMaterialInterface* BaseMaterial = InBaseMaterial;
+	
+	if (!IsValidFlattenMaterial(InBaseMaterial))
+	{
+		BaseMaterial = GEngine->DefaultFlattenMaterial;
+	} 
+
+	UMaterialInstanceConstant* OutMaterial = FMaterialUtilities::CreateInstancedMaterial(BaseMaterial, InOuter, AssetBasePath + AssetBaseName, RF_Public | RF_Standalone);
+	OutAssetsToSync.Add(OutMaterial);
+
+	OutMaterial->BasePropertyOverrides.TwoSided = FlattenMaterial.bTwoSided && InMaterialProxySettings.bAllowTwoSidedMaterial;
+	OutMaterial->BasePropertyOverrides.bOverride_TwoSided = (FlattenMaterial.bTwoSided != false) && InMaterialProxySettings.bAllowTwoSidedMaterial;
+	OutMaterial->BasePropertyOverrides.DitheredLODTransition = FlattenMaterial.bDitheredLODTransition;
+	OutMaterial->BasePropertyOverrides.bOverride_DitheredLODTransition = FlattenMaterial.bDitheredLODTransition != false;
+
+	if (InMaterialProxySettings.BlendMode != BLEND_Opaque)
+	{
+		OutMaterial->BasePropertyOverrides.bOverride_BlendMode = true;
+		OutMaterial->BasePropertyOverrides.BlendMode = InMaterialProxySettings.BlendMode;
+	}
+
+	bool bPackMetallic, bPackSpecular, bPackRoughness;
+	int32 NumSamples = 0;
+	FIntPoint PackedSize;
+	const bool bPackTextures = CalculatePackedTextureData(FlattenMaterial, bPackMetallic, bPackSpecular, bPackRoughness, NumSamples, PackedSize);
+
+	FStaticParameterSet NewStaticParameterSet;
+
+	if(FlattenMaterial.UVChannel != 0)
+	{
+		// If the used texture coordinate was not the default UV0 set the appropriate one on the instance material
+		FStaticSwitchParameter SwitchParameter;
+		SwitchParameter.ParameterInfo.Name = TEXT("UseCustomUV");
+		SwitchParameter.Value = true;
+		SwitchParameter.bOverride = true;
+		NewStaticParameterSet.StaticSwitchParameters.Add(SwitchParameter);
+
+		SwitchParameter.ParameterInfo.Name = *(TEXT("UseUV") + FString::FromInt(FlattenMaterial.UVChannel));
+		NewStaticParameterSet.StaticSwitchParameters.Add(SwitchParameter);
+	}
+
+	auto CreateTextureFromDefault = [&](const FName TextureName, const FString& AssetLongName, FIntPoint Size, const TArray<FColor>& Samples)
+	{
+		FCreateTexture2DParameters CreateParams;
+		CreateParams.bSRGB = false;
+		CreateParams.bVirtualTexture = false;
+		CreateParams.CompressionSettings = TextureCompressionSettings::TC_Default;
+		CreateParams.TextureGroup = TextureGroup::TEXTUREGROUP_World;
+
+		UTexture* DefaultTexture = nullptr;
+		OutMaterial->GetTextureParameterValue(TextureName, DefaultTexture);
+		if (ensure(DefaultTexture))
+		{
+			CreateParams.bSRGB = DefaultTexture->SRGB;
+			CreateParams.bVirtualTexture = DefaultTexture->VirtualTextureStreaming;
+			CreateParams.CompressionSettings = DefaultTexture->CompressionSettings;
+			CreateParams.TextureGroup = DefaultTexture->LODGroup;
+		}
+
+		return FMaterialUtilities::CreateTexture(InOuter, AssetLongName, Size, Samples, CreateParams, RF_Public | RF_Standalone);
+	};
+
+	auto SetTextureParam = [&](EFlattenMaterialProperties FlattenProperty)
+	{
+		if (FlattenMaterial.DoesPropertyContainData(FlattenProperty) && !FlattenMaterial.IsPropertyConstant(FlattenProperty))
+		{
+			FString PropertyName = GetMatchingParamName(FlattenProperty, OutMaterial);
+
+			const FName TextureName(PropertyName + TEXT("Texture"));
+			const FName UseTexture(TEXT("Use") + PropertyName);
+
+			UTexture2D* Texture = CreateTextureFromDefault(TextureName, AssetBasePath + TEXT("T_") + AssetBaseName + TEXT("_") + PropertyName, FlattenMaterial.GetPropertySize(FlattenProperty), FlattenMaterial.GetPropertySamples(FlattenProperty));
+
+			FStaticSwitchParameter SwitchParameter;
+			SwitchParameter.ParameterInfo.Name = UseTexture;
+			SwitchParameter.Value = true;
+			SwitchParameter.bOverride = true;
+			NewStaticParameterSet.StaticSwitchParameters.Add(SwitchParameter);
+
+			OutMaterial->SetTextureParameterValueEditorOnly(TextureName, Texture);
+			
+			OutAssetsToSync.Add(Texture);
+		}
+	};
+
+	auto SetTextureParamConstVector = [&](EFlattenMaterialProperties FlattenProperty)
+	{
+		if (FlattenMaterial.DoesPropertyContainData(FlattenProperty) && !FlattenMaterial.IsPropertyConstant(FlattenProperty))
+		{
+			SetTextureParam(FlattenProperty);
+		} 
+		else
+		{
+			FString PropertyName = GetMatchingParamName(FlattenProperty, OutMaterial);
+
+			const FName ConstName(PropertyName + TEXT("Const"));
+			OutMaterial->SetVectorParameterValueEditorOnly(ConstName, FlattenMaterial.GetPropertySamples(FlattenProperty)[0]);
+		} 
+	};
+
+	auto SetTextureParamConstScalar = [&](EFlattenMaterialProperties FlattenProperty, float ConstantValue)
+	{
+		if (FlattenMaterial.DoesPropertyContainData(FlattenProperty) && !FlattenMaterial.IsPropertyConstant(FlattenProperty))
+		{
+			SetTextureParam(FlattenProperty);
+		}
+		else
+		{
+			FString PropertyName = GetMatchingParamName(FlattenProperty, OutMaterial);
+
+			const FName ConstName(PropertyName + TEXT("Const"));
+			FLinearColor Colour = FlattenMaterial.IsPropertyConstant(FlattenProperty) ? FLinearColor::FromSRGBColor(FlattenMaterial.GetPropertySamples(FlattenProperty)[0]) : FLinearColor(ConstantValue, 0, 0, 0);
+			OutMaterial->SetScalarParameterValueEditorOnly(ConstName, Colour.R);
+		}
+	};
+
+	// Load textures and set switches accordingly
+	if (FlattenMaterial.GetPropertySize(EFlattenMaterialProperties::Diffuse).Num() > 0 && !(FlattenMaterial.IsPropertyConstant(EFlattenMaterialProperties::Diffuse) && FlattenMaterial.GetPropertySamples(EFlattenMaterialProperties::Diffuse)[0] == FColor::Black))
+	{
+		SetTextureParamConstVector(EFlattenMaterialProperties::Diffuse);
+	}
+
+	if (FlattenMaterial.GetPropertySize(EFlattenMaterialProperties::Normal).Num() > 1)
+	{
+		SetTextureParam(EFlattenMaterialProperties::Normal);
+	}
+
+	// Determine whether or not specific material properties are packed together into one texture (requires at least two to match (number of samples and texture size) in order to be packed
+	if (!bPackMetallic && (FlattenMaterial.GetPropertySize(EFlattenMaterialProperties::Metallic).Num() > 0 || !InMaterialProxySettings.bMetallicMap))
+	{
+		SetTextureParamConstScalar(EFlattenMaterialProperties::Metallic, InMaterialProxySettings.MetallicConstant);
+	}
+
+	if (!bPackRoughness && (FlattenMaterial.GetPropertySize(EFlattenMaterialProperties::Roughness).Num() > 0 || !InMaterialProxySettings.bRoughnessMap))
+	{
+		SetTextureParamConstScalar(EFlattenMaterialProperties::Roughness, InMaterialProxySettings.RoughnessConstant);
+	}
+
+	if (!bPackSpecular && (FlattenMaterial.GetPropertySize(EFlattenMaterialProperties::Specular).Num() > 0 || !InMaterialProxySettings.bSpecularMap))
+	{
+		SetTextureParamConstScalar(EFlattenMaterialProperties::Specular, InMaterialProxySettings.SpecularConstant);
+	}
+
+	if (FlattenMaterial.GetPropertySize(EFlattenMaterialProperties::Opacity).Num() > 0 || !InMaterialProxySettings.bOpacityMap)
+	{
+		SetTextureParamConstScalar(EFlattenMaterialProperties::Opacity, InMaterialProxySettings.OpacityConstant);
+	}
+
+	if (FlattenMaterial.GetPropertySize(EFlattenMaterialProperties::OpacityMask).Num() > 0 || !InMaterialProxySettings.bOpacityMaskMap)
+	{
+		SetTextureParamConstScalar(EFlattenMaterialProperties::OpacityMask, InMaterialProxySettings.OpacityMaskConstant);
+	}
+
+	if (FlattenMaterial.GetPropertySize(EFlattenMaterialProperties::AmbientOcclusion).Num() > 0 || !InMaterialProxySettings.bAmbientOcclusionMap)
+	{
+		SetTextureParamConstScalar(EFlattenMaterialProperties::AmbientOcclusion, InMaterialProxySettings.AmbientOcclusionConstant);
+	}
+
+	// Handle the packed texture if applicable
+	if (bPackTextures)
+	{
+		TArray<FColor> MergedTexture;
+		MergedTexture.AddZeroed(NumSamples);
+
+		// Merge properties into one texture using the separate colour channels
+		const EFlattenMaterialProperties Properties[3] = { EFlattenMaterialProperties::Metallic , EFlattenMaterialProperties::Roughness, EFlattenMaterialProperties::Specular };
+
+		//Property that is not part of the pack (because of a different size), will see is reserve pack space fill with Black Color.
+		const bool PropertyShouldBePack[3] = { bPackMetallic , bPackRoughness , bPackSpecular };
+
+		// Red mask (all properties are rendered into the red channel)
+		FColor NonAlphaRed = FColor::Red;
+		NonAlphaRed.A = 0;
+		const uint32 ColorMask = NonAlphaRed.DWColor();
+		const uint32 Shift[3] = { 0, 8, 16 };
+		for (int32 PropertyIndex = 0; PropertyIndex < 3; ++PropertyIndex)
+		{
+			const EFlattenMaterialProperties Property = Properties[PropertyIndex];
+			const bool HasProperty = PropertyShouldBePack[PropertyIndex] && FlattenMaterial.DoesPropertyContainData(Property) && !FlattenMaterial.IsPropertyConstant(Property);
+
+			if (HasProperty)
+			{
+				const TArray<FColor>& PropertySamples = FlattenMaterial.GetPropertySamples(Property);
+				// OR masked values (samples initialized to zero, so no random data)
+				for (int32 SampleIndex = 0; SampleIndex < NumSamples; ++SampleIndex)
+				{
+					// Black adds the alpha + red channel value shifted into the correct output channel
+					MergedTexture[SampleIndex].DWColor() |= (FColor::Black.DWColor() + ((PropertySamples[SampleIndex].DWColor() & ColorMask) >> Shift[PropertyIndex]));
+				}
+			}
+		}
+
+		// Create texture using the merged property data
+		const FName PackedTextureName(TEXT("PackedTexture"));
+		UTexture2D* PackedTexture = CreateTextureFromDefault(PackedTextureName, AssetBasePath + TEXT("T_") + AssetBaseName + TEXT("_MRS"), PackedSize, MergedTexture);
+		checkf(PackedTexture, TEXT("Failed to create texture"));
+		OutAssetsToSync.Add(PackedTexture);
+
+		// Setup switches for whether or not properties will be packed into one texture
+		FStaticSwitchParameter SwitchParameter;
+		SwitchParameter.ParameterInfo.Name = TEXT("PackMetallic");
+		SwitchParameter.Value = bPackMetallic;
+		SwitchParameter.bOverride = true;
+		NewStaticParameterSet.StaticSwitchParameters.Add(SwitchParameter);
+
+		SwitchParameter.ParameterInfo.Name = TEXT("PackSpecular");
+		SwitchParameter.Value = bPackSpecular;
+		SwitchParameter.bOverride = true;
+		NewStaticParameterSet.StaticSwitchParameters.Add(SwitchParameter);
+
+		SwitchParameter.ParameterInfo.Name = TEXT("PackRoughness");
+		SwitchParameter.Value = bPackRoughness;
+		SwitchParameter.bOverride = true;
+		NewStaticParameterSet.StaticSwitchParameters.Add(SwitchParameter);
+
+		// Set up switch and texture values
+		OutMaterial->SetTextureParameterValueEditorOnly(PackedTextureName, PackedTexture);
+	}
+
+	// Emissive is a special case due to the scaling variable
+	if (FlattenMaterial.GetPropertySamples(EFlattenMaterialProperties::Emissive).Num() > 0 && !(FlattenMaterial.GetPropertySize(EFlattenMaterialProperties::Emissive).Num() == 1 && FlattenMaterial.GetPropertySamples(EFlattenMaterialProperties::Emissive)[0] == FColor::Black))
+	{
+		SetTextureParamConstVector(EFlattenMaterialProperties::Emissive);
+
+		if (FlattenMaterial.EmissiveScale != 1.0f)
+		{
+			FMaterialParameterInfo ParameterInfo(TEXT("EmissiveScale"));
+			OutMaterial->SetScalarParameterValueEditorOnly(ParameterInfo, FlattenMaterial.EmissiveScale);
+		}
+	}
+
+	// Force initializing the static permutations according to the switches we have set
+	OutMaterial->UpdateStaticPermutation(NewStaticParameterSet, MaterialUpdateContext);
+	OutMaterial->InitStaticPermutation();
+
+	OutMaterial->PostEditChange();
+
+	return OutMaterial;
+}	
+
+#undef LOCTEXT_NAMESPACE // "ProxyMaterialUtilities"
\ No newline at end of file
diff --git a/Engine/Source/Developer/MaterialUtilities/Public/MaterialUtilities.h b/Engine/Source/Developer/MaterialUtilities/Public/MaterialUtilities.h
--- a/Engine/Source/Developer/MaterialUtilities/Public/MaterialUtilities.h
+++ b/Engine/Source/Developer/MaterialUtilities/Public/MaterialUtilities.h
@@ -13,6 +13,7 @@
 
 #include "LightMap.h"
 #include "ShadowMap.h"
+#include "ImageUtils.h"
 
 class ALandscapeProxy;
 class Error;
@@ -266,6 +267,15 @@ public:
 	 */
 	static bool SupportsExport(EBlendMode InBlendMode, EMaterialProperty InMaterialProperty);
 
+	/**
+	 * Flattens specified landscape material
+	 *
+	 * @param InLandscape			Target landscape
+	 * @param OutFlattenMaterial	Output flattened material
+	 * @return						Whether operation was successful
+	 */
+	static bool ExportLandscapeMaterial(const ALandscapeProxy* InLandscape, FFlattenMaterial& OutFlattenMaterial);
+
 	/**
 	 * Flattens specified landscape material
 	 *
@@ -274,7 +284,7 @@ public:
 	 * @param OutFlattenMaterial	Output flattened material
 	 * @return						Whether operation was successful
 	 */
-	static bool ExportLandscapeMaterial(ALandscapeProxy* InLandscape, const TSet<FPrimitiveComponentId>& HiddenPrimitives, FFlattenMaterial& OutFlattenMaterial);
+	static bool ExportLandscapeMaterial(const ALandscapeProxy* InLandscape, const TSet<FPrimitiveComponentId>& HiddenPrimitives, FFlattenMaterial& OutFlattenMaterial);
 	
 	/**
  	 * Generates a texture from an array of samples 
@@ -292,6 +302,19 @@ public:
 	 */
 	static UTexture2D* CreateTexture(UPackage* Outer, const FString& AssetLongName, FIntPoint Size, const TArray<FColor>& Samples, TextureCompressionSettings CompressionSettings, TextureGroup LODGroup, EObjectFlags Flags, bool bSRGB, const FGuid& SourceGuidHash = FGuid());
 
+	/**
+	 * Generates a texture from an array of samples
+	 *
+	 * @param Outer					Outer for the material and texture objects, if NULL new packages will be created for each asset
+	 * @param AssetLongName			Long asset path for the new texture
+	 * @param Size					Resolution of the texture to generate (must match the number of samples)
+	 * @param Samples				Color data for the texture
+	 * @param CreateParams			Params about how to set up the texture
+	 * @param Flags					ObjectFlags for the new texture
+	 * @return						The new texture.
+	 */
+	static UTexture2D* CreateTexture(UPackage* Outer, const FString& AssetLongName, FIntPoint Size, const TArray<FColor>& Samples, const FCreateTexture2DParameters& CreateParams, EObjectFlags Flags);
+
 	/**
 	 * Creates UMaterial object from a flatten material
 	 *
@@ -449,6 +472,9 @@ public:
 	/** Creates a proxy material and the required texture assets */
 	static UMaterialInterface* CreateProxyMaterialAndTextures(const FString& PackageName, const FString& AssetName, const FBakeOutput& BakeOutput, const FMeshData& MeshData, const FMaterialData& MaterialData, UMaterialOptions* Options);
 
+	/** Creates a flatten material instance and the required texture assets */
+	static UMaterialInstanceConstant* CreateFlattenMaterialInstance(UPackage* InOuter, const FMaterialProxySettings& InMaterialProxySettings, UMaterialInterface* InBaseMaterial, const FFlattenMaterial& FlattenMaterial, const FString& AssetBasePath, const FString& AssetBaseName, TArray<UObject*>& OutAssetsToSync, FMaterialUpdateContext* MaterialUpdateContext = nullptr);
+
 	/** Compute the required texel density needed to properly represent an object/objects covering the provided world space radius
 	* 
 	* @param InScreenSize		Screen size at which the objects are expected to be be shown
@@ -465,6 +491,9 @@ public:
 	*/
 	static float ComputeRequiredTexelDensityFromDrawDistance(const float InDrawDistance, float InWorldSpaceRadius);
 
+	/** Validate that the provided material has all the required parameters needed to be considered a flattening material */
+	static bool IsValidFlattenMaterial(const UMaterialInterface* InBaseMaterial);
+
 private:
 	
 	/**
diff --git a/Engine/Source/Developer/MeshMergeUtilities/Private/MeshMergeUtilities.cpp b/Engine/Source/Developer/MeshMergeUtilities/Private/MeshMergeUtilities.cpp
--- a/Engine/Source/Developer/MeshMergeUtilities/Private/MeshMergeUtilities.cpp
+++ b/Engine/Source/Developer/MeshMergeUtilities/Private/MeshMergeUtilities.cpp
@@ -10,6 +10,7 @@
 
 #include "Misc/PackageName.h"
 #include "MaterialUtilities.h"
+#include "Materials/MaterialInstanceConstant.h"
 #include "Components/SkeletalMeshComponent.h"
 #include "Components/SplineMeshComponent.h"
 #include "Components/SkinnedMeshComponent.h"
@@ -35,7 +36,6 @@
 #include "ProxyGenerationProcessor.h"
 #include "Editor/EditorPerProjectUserSettings.h"
 
-#include "ProxyMaterialUtilities.h"
 #include "Engine/StaticMesh.h"
 #include "PhysicsEngine/ConvexElem.h"
 #include "PhysicsEngine/BodySetup.h"
@@ -74,6 +74,9 @@
 #include "Async/Async.h"
 #include "TextureCompiler.h"
 
+#include "Widgets/Notifications/SNotificationList.h"
+#include "Framework/Notifications/NotificationManager.h"
+
 #include "ISMPartition/ISMComponentDescriptor.h"
 
 #define LOCTEXT_NAMESPACE "MeshMergeUtils"
@@ -1961,7 +1964,18 @@ void FMeshMergeUtilities::CreateProxyMesh(const TArray<UStaticMeshComponent*>& I
 
 bool FMeshMergeUtilities::IsValidBaseMaterial(const UMaterialInterface* InBaseMaterial, bool bShowToaster) const
 {
-	return ProxyMaterialUtilities::IsValidBaseMaterial(InBaseMaterial, bShowToaster);
+	bool bIsValid = FMaterialUtilities::IsValidFlattenMaterial(InBaseMaterial);
+	if (!bIsValid && bShowToaster)
+	{
+		FFormatNamedArguments Arguments;
+		Arguments.Add(TEXT("MaterialName"), FText::FromString(InBaseMaterial->GetName()));
+		FText ErrorMessage = FText::Format(LOCTEXT("UHierarchicalLODSettings_PostEditChangeProperty", "Material {MaterialName} is missing required Material Parameters (check log for details)"), Arguments);
+		FNotificationInfo Info(ErrorMessage);
+		Info.ExpireDuration = 5.0f;
+		FSlateNotificationManager::Get().AddNotification(Info);
+	}
+
+	return bIsValid;
 }
 
 void FMeshMergeUtilities::RetrieveMeshDescription(const UStaticMeshComponent* InStaticMeshComponent, int32 LODIndex, FMeshDescription& InOutMeshDescription, bool bPropagateMeshData) const
@@ -3587,7 +3601,7 @@ UMaterialInterface* FMeshMergeUtilities::CreateProxyMaterial(const FString &InBa
 		MaterialPackage->Modify();
 	}
 
-	UMaterialInstanceConstant* MergedMaterial = ProxyMaterialUtilities::CreateProxyMaterialInstance(MaterialPackage, InSettings.MaterialSettings, InBaseMaterial, OutMaterial, MaterialPackageName, MaterialAssetName, OutAssetsToSync, InMaterialUpdateContext);
+	UMaterialInstanceConstant* MergedMaterial = FMaterialUtilities::CreateFlattenMaterialInstance(MaterialPackage, InSettings.MaterialSettings, InBaseMaterial, OutMaterial, MaterialPackageName, MaterialAssetName, OutAssetsToSync, InMaterialUpdateContext);
 	// Set material static lighting usage flag if project has static lighting enabled
 	static const auto AllowStaticLightingVar = IConsoleManager::Get().FindTConsoleVariableDataInt(TEXT("r.AllowStaticLighting"));
 	const bool bAllowStaticLighting = (!AllowStaticLightingVar || AllowStaticLightingVar->GetValueOnGameThread() != 0);
diff --git a/Engine/Source/Developer/MeshMergeUtilities/Private/MeshMergeUtilities.h b/Engine/Source/Developer/MeshMergeUtilities/Private/MeshMergeUtilities.h
--- a/Engine/Source/Developer/MeshMergeUtilities/Private/MeshMergeUtilities.h
+++ b/Engine/Source/Developer/MeshMergeUtilities/Private/MeshMergeUtilities.h
@@ -36,6 +36,7 @@ public:
 	UE_DEPRECATED(4.20, "Use different signature containing BaseMaterial parameter")
 	virtual void CreateProxyMesh(const TArray<UStaticMeshComponent*>& InStaticMeshComps, const struct FMeshProxySettings& InMeshProxySettings, UPackage* InOuter, const FString& InProxyBasePackageName, const FGuid InGuid, const FCreateProxyDelegate& InProxyCreatedDelegate, const bool bAllowAsync = false, const float ScreenSize = 1.0f) const override;
 
+	UE_DEPRECATED(5.0, "Use FMaterialUtilities::IsValidFlattenMaterial()")
 	virtual bool IsValidBaseMaterial(const UMaterialInterface* InBaseMaterial, bool bShowToaster) const override;
 
 	virtual void RetrieveMeshDescription(const UStaticMeshComponent* InStaticMeshComponent, int32 LODIndex, FMeshDescription& InOutMeshDescription, bool bPropagateVertexColours) const override;
diff --git a/Engine/Source/Developer/MeshMergeUtilities/Private/ProxyGenerationProcessor.cpp b/Engine/Source/Developer/MeshMergeUtilities/Private/ProxyGenerationProcessor.cpp
--- a/Engine/Source/Developer/MeshMergeUtilities/Private/ProxyGenerationProcessor.cpp
+++ b/Engine/Source/Developer/MeshMergeUtilities/Private/ProxyGenerationProcessor.cpp
@@ -4,7 +4,6 @@
 #include "MaterialUtilities.h"
 #include "MeshMergeUtilities.h"
 #include "IMeshMergeExtension.h"
-#include "ProxyMaterialUtilities.h"
 #include "IMeshReductionInterfaces.h"
 #include "IMeshReductionManagerModule.h"
 #include "Modules/ModuleManager.h"
@@ -16,6 +15,7 @@
 #include "Editor.h"
 #include "MeshMergeHelpers.h"
 #include "ObjectCacheEventSink.h"
+#include "Materials/MaterialInstanceConstant.h"
 #endif // WITH_EDITOR
 
 FProxyGenerationProcessor::FProxyGenerationProcessor(const FMeshMergeUtilities* InOwner)
@@ -160,7 +160,7 @@ void FProxyGenerationProcessor::ProcessJob(const FGuid& JobGuid, FProxyGeneratio
 		FMaterialUtilities::OptimizeFlattenMaterial(FlattenMaterial);
 
 		// Create a new proxy material instance
-		ProxyMaterial = ProxyMaterialUtilities::CreateProxyMaterialInstance(Data->MergeData->InOuter, Data->MergeData->InProxySettings.MaterialSettings, Data->MergeData->BaseMaterial, FlattenMaterial, AssetBasePath, AssetBaseName, OutAssetsToSync, &MaterialUpdateContext);
+		ProxyMaterial = FMaterialUtilities::CreateFlattenMaterialInstance(Data->MergeData->InOuter, Data->MergeData->InProxySettings.MaterialSettings, Data->MergeData->BaseMaterial, FlattenMaterial, AssetBasePath, AssetBaseName, OutAssetsToSync, &MaterialUpdateContext);
 
 		for (IMeshMergeExtension* Extension : Owner->MeshMergeExtensions)
 		{
diff --git a/Engine/Source/Developer/MeshMergeUtilities/Private/ProxyMaterialUtilities.h b/Engine/Source/Developer/MeshMergeUtilities/Private/ProxyMaterialUtilities.h
deleted file mode 100644
diff --git a/Engine/Source/Developer/MeshMergeUtilities/Public/IMeshMergeUtilities.h b/Engine/Source/Developer/MeshMergeUtilities/Public/IMeshMergeUtilities.h
--- a/Engine/Source/Developer/MeshMergeUtilities/Public/IMeshMergeUtilities.h
+++ b/Engine/Source/Developer/MeshMergeUtilities/Public/IMeshMergeUtilities.h
@@ -125,7 +125,7 @@ public:
 	*/
 	virtual void CreateProxyMesh(const TArray<UStaticMeshComponent*>& InComponents, const struct FMeshProxySettings& InMeshProxySettings, UMaterialInterface* InBaseMaterial, UPackage* InOuter, const FString& InProxyBasePackageName, const FGuid InGuid, const FCreateProxyDelegate& InProxyCreatedDelegate, const bool bAllowAsync = false, const float ScreenSize = 1.0f) const = 0;
 
-	/** Checks whether or not the give material is valid as a base for the Proxy Material. This as it requires certain Material Parameters to be included to get a desirable material */
+	UE_DEPRECATED(5.0, "Use FMaterialUtilities::IsValidFlattenMaterial()")
 	virtual bool IsValidBaseMaterial(const UMaterialInterface* InBaseMaterial, bool bShowToaster) const = 0;
 
 	/** Extract mesh data in FMeshDescription form from the provided static mesh component */
diff --git a/Engine/Source/Editor/UnrealEd/Private/HierarchicalLOD.cpp b/Engine/Source/Editor/UnrealEd/Private/HierarchicalLOD.cpp
--- a/Engine/Source/Editor/UnrealEd/Private/HierarchicalLOD.cpp
+++ b/Engine/Source/Editor/UnrealEd/Private/HierarchicalLOD.cpp
@@ -30,10 +30,14 @@
 #include "Editor.h"
 #include "UnrealEdGlobals.h"
 #include "HLOD/HLODEngineSubsystem.h"
-#include "IMeshMergeUtilities.h"
+#include "MaterialUtilities.h"
 #include "MeshMergeModule.h"
 #include "MeshDescription.h"
 #include "StaticMeshOperations.h"
+
+#include "Widgets/Notifications/SNotificationList.h"
+#include "Framework/Notifications/NotificationManager.h"
+
 #endif // WITH_EDITOR
 
 
@@ -60,14 +64,32 @@ UHierarchicalLODSettings::UHierarchicalLODSettings(const FObjectInitializer& Obj
 	}
 }
 
+bool UHierarchicalLODSettings::IsValidFlattenMaterial(const UMaterialInterface* InBaseMaterial, bool bShowToaster)
+{
+	bool bIsValid = FMaterialUtilities::IsValidFlattenMaterial(InBaseMaterial);
+
+#if WITH_EDITOR
+	if (!bIsValid && bShowToaster)
+	{
+		FFormatNamedArguments Arguments;
+		Arguments.Add(TEXT("MaterialName"), FText::FromString(InBaseMaterial->GetName()));
+		FText ErrorMessage = FText::Format(LOCTEXT("UHierarchicalLODSettings_PostEditChangeProperty", "Material {MaterialName} is missing required Material Parameters (check log for details)"), Arguments);
+		FNotificationInfo Info(ErrorMessage);
+		Info.ExpireDuration = 5.0f;
+		FSlateNotificationManager::Get().AddNotification(Info);
+	}
+#endif
+
+	return bIsValid;
+}
+
 void UHierarchicalLODSettings::PostEditChangeProperty(struct FPropertyChangedEvent& PropertyChangedEvent)
 {
 	if (PropertyChangedEvent.GetPropertyName() == GET_MEMBER_NAME_CHECKED(UHierarchicalLODSettings, BaseMaterial))
 	{
 		if (!BaseMaterial.IsNull())
 		{
-			const IMeshMergeUtilities& Module = FModuleManager::Get().LoadModuleChecked<IMeshMergeModule>("MeshMergeUtilities").GetUtilities();
-			if (!Module.IsValidBaseMaterial(BaseMaterial.LoadSynchronous(), true))
+			if (!IsValidFlattenMaterial(BaseMaterial.LoadSynchronous(), true))
 			{
 				BaseMaterial = GEngine->DefaultFlattenMaterial;
 			}
diff --git a/Engine/Source/Editor/UnrealEd/Public/HierarchicalLOD.h b/Engine/Source/Editor/UnrealEd/Public/HierarchicalLOD.h
--- a/Engine/Source/Editor/UnrealEd/Public/HierarchicalLOD.h
+++ b/Engine/Source/Editor/UnrealEd/Public/HierarchicalLOD.h
@@ -53,6 +53,8 @@ public:
 	TSoftObjectPtr<class UMaterialInterface> BaseMaterial;
 
 #if WITH_EDITOR
+	static bool IsValidFlattenMaterial(const UMaterialInterface* InBaseMaterial, bool bShowToaster);
+
 	virtual void PostEditChangeProperty(struct FPropertyChangedEvent& PropertyChangedEvent) override;
 #endif // WITH_EDITOR
 };
diff --git a/Engine/Source/Runtime/Engine/Engine.Build.cs b/Engine/Source/Runtime/Engine/Engine.Build.cs
--- a/Engine/Source/Runtime/Engine/Engine.Build.cs
+++ b/Engine/Source/Runtime/Engine/Engine.Build.cs
@@ -238,11 +238,14 @@ public class Engine : ModuleRules
 
 		if (Target.bCompileAgainstEditor)
 		{
+			PrivateIncludePathModuleNames.Add("MaterialUtilities");
+			PrivateDependencyModuleNames.Add("MaterialUtilities");
+
 			PrivateIncludePathModuleNames.Add("MeshUtilities");
+			DynamicallyLoadedModuleNames.Add("MeshUtilities");
+
 			PrivateIncludePathModuleNames.Add("MeshUtilitiesCommon");
 
-			DynamicallyLoadedModuleNames.Add("MeshUtilities");
-			
 			PublicIncludePathModuleNames.Add("AnimationDataController");
 			DynamicallyLoadedModuleNames.Add("AnimationDataController");
 
diff --git a/Engine/Source/Runtime/Engine/Private/ImageUtils.cpp b/Engine/Source/Runtime/Engine/Private/ImageUtils.cpp
--- a/Engine/Source/Runtime/Engine/Private/ImageUtils.cpp
+++ b/Engine/Source/Runtime/Engine/Private/ImageUtils.cpp
@@ -329,6 +329,8 @@ UTexture2D* FImageUtils::CreateTexture2D(int32 SrcWidth, int32 SrcHeight, const
 		Tex2D->LODGroup = InParams.TextureGroup;
 	}
 
+	Tex2D->VirtualTextureStreaming = InParams.bVirtualTexture;
+
 	Tex2D->PostEditChange();
 	return Tex2D;
 #else
diff --git a/Engine/Source/Runtime/Engine/Private/WorldSettings.cpp b/Engine/Source/Runtime/Engine/Private/WorldSettings.cpp
--- a/Engine/Source/Runtime/Engine/Private/WorldSettings.cpp
+++ b/Engine/Source/Runtime/Engine/Private/WorldSettings.cpp
@@ -36,8 +36,6 @@
 #if WITH_EDITOR
 #include "Editor.h"
 #include "HierarchicalLOD.h"
-#include "IMeshMergeUtilities.h"
-#include "MeshMergeModule.h"
 #include "Settings/EditorExperimentalSettings.h"
 #include "Landscape.h"
 #include "Rendering/StaticLightingSystemInterface.h"
@@ -774,8 +772,7 @@ void AWorldSettings::InternalPostPropertyChanged(FName PropertyName)
 		{
 			if (!OverrideBaseMaterial.IsNull())
 			{
-				const IMeshMergeUtilities& Module = FModuleManager::Get().LoadModuleChecked<IMeshMergeModule>("MeshMergeUtilities").GetUtilities();
-				if (!Module.IsValidBaseMaterial(OverrideBaseMaterial.LoadSynchronous(), true))
+				if (!UHierarchicalLODSettings::IsValidFlattenMaterial(OverrideBaseMaterial.LoadSynchronous(), true))
 				{
 					OverrideBaseMaterial = GEngine->DefaultHLODFlattenMaterial;
 				}
@@ -790,8 +787,7 @@ void UHierarchicalLODSetup::PostEditChangeProperty(struct FPropertyChangedEvent&
 	{
 		if (!OverrideBaseMaterial.IsNull())
 		{
-			const IMeshMergeUtilities& Module = FModuleManager::Get().LoadModuleChecked<IMeshMergeModule>("MeshMergeUtilities").GetUtilities();
-			if (!Module.IsValidBaseMaterial(OverrideBaseMaterial.LoadSynchronous(), true))
+			if (!UHierarchicalLODSettings::IsValidFlattenMaterial(OverrideBaseMaterial.LoadSynchronous(), true))
 			{
 				OverrideBaseMaterial = GEngine->DefaultHLODFlattenMaterial;
 			}
diff --git a/Engine/Source/Runtime/Engine/Public/ImageUtils.h b/Engine/Source/Runtime/Engine/Public/ImageUtils.h
--- a/Engine/Source/Runtime/Engine/Public/ImageUtils.h
+++ b/Engine/Source/Runtime/Engine/Public/ImageUtils.h
@@ -16,7 +16,7 @@ class UTextureRenderTarget2D;
 class UTextureRenderTargetCube;
 
 /**
- *	Parameters used for creating a Texture2D frmo a simple color buffer.
+ *	Parameters used for creating a Texture2D from a simple color buffer.
  */
 struct FCreateTexture2DParameters
 {
@@ -32,6 +32,9 @@ struct FCreateTexture2DParameters
 	/** If texture should be set as SRGB */
 	bool						bSRGB;
 
+	/** If texture should be a virtual texture */
+	bool						bVirtualTexture;
+
 	/** Mip-map generation settings */
 	TextureMipGenSettings		MipGenSettings;
 
@@ -46,6 +49,7 @@ struct FCreateTexture2DParameters
 			CompressionSettings(TC_Default),
 			bDeferCompression(false),
 			bSRGB(true),
+			bVirtualTexture(false),
 			MipGenSettings(TMGS_FromTextureGroup),
 			TextureGroup(TEXTUREGROUP_MAX)
 	{
