Fixed static analysis warning

diff --git a/Engine/Source/Developer/MaterialUtilities/Private/MaterialUtilities.cpp b/Engine/Source/Developer/MaterialUtilities/Private/MaterialUtilities.cpp
--- a/Engine/Source/Developer/MaterialUtilities/Private/MaterialUtilities.cpp
+++ b/Engine/Source/Developer/MaterialUtilities/Private/MaterialUtilities.cpp
@@ -82,6 +82,8 @@ void FMaterialUtilities::OnPreGarbageCollect()
 
 UMaterialInterface* FMaterialUtilities::CreateProxyMaterialAndTextures(UPackage* OuterPackage, const FString& AssetName, const FBakeOutput& BakeOutput, const FMeshData& MeshData, const FMaterialData& MaterialData, UMaterialOptions* Options)
 {
+	check(MaterialData.Material);
+
 	TArray<EMaterialProperty> SRGBEnabledProperties{ MP_BaseColor, MP_EmissiveColor, MP_SubsurfaceColor };
 
 	// Certain material properties use differen compression settings
