- Fixing non unity mac compilation

diff --git a/Engine/Source/Developer/MaterialUtilities/Private/ProxyMaterialUtilities.cpp b/Engine/Source/Developer/MaterialUtilities/Private/ProxyMaterialUtilities.cpp
--- a/Engine/Source/Developer/MaterialUtilities/Private/ProxyMaterialUtilities.cpp
+++ b/Engine/Source/Developer/MaterialUtilities/Private/ProxyMaterialUtilities.cpp
@@ -1,7 +1,5 @@
 // Copyright Epic Games, Inc. All Rights Reserved.
 
-#pragma once
-
 #include "CoreMinimal.h"
 #include "Engine/MaterialMerging.h"
 #include "StaticParameterSet.h"
