Fix metallic, specular, and roughness packed into a single texture.

diff --git a/Engine/Source/Developer/MaterialUtilities/Private/MaterialUtilities.cpp b/Engine/Source/Developer/MaterialUtilities/Private/MaterialUtilities.cpp
--- a/Engine/Source/Developer/MaterialUtilities/Private/MaterialUtilities.cpp
+++ b/Engine/Source/Developer/MaterialUtilities/Private/MaterialUtilities.cpp
@@ -1057,11 +1057,10 @@ UMaterial* FMaterialUtilities::CreateMaterial(const FFlattenMaterial& InFlattenM
 		MaterialNodeY += MaterialNodeStepY;
 	}
 
-
 	// Whether or not a material property is baked down
 	const bool bHasMetallic = InFlattenMaterial.DoesPropertyContainData(EFlattenMaterialProperties::Metallic) && !InFlattenMaterial.IsPropertyConstant(EFlattenMaterialProperties::Metallic);
-	const bool bHasRoughness = InFlattenMaterial.DoesPropertyContainData(EFlattenMaterialProperties::Roughness) && !InFlattenMaterial.IsPropertyConstant(EFlattenMaterialProperties::Roughness);
 	const bool bHasSpecular = InFlattenMaterial.DoesPropertyContainData(EFlattenMaterialProperties::Specular) && !InFlattenMaterial.IsPropertyConstant(EFlattenMaterialProperties::Specular);
+	const bool bHasRoughness = InFlattenMaterial.DoesPropertyContainData(EFlattenMaterialProperties::Roughness) && !InFlattenMaterial.IsPropertyConstant(EFlattenMaterialProperties::Roughness);
 
 	// Number of material properties baked down to textures
 	const int BakedMaterialPropertyCount = bHasMetallic + bHasRoughness + bHasSpecular;
@@ -1080,22 +1079,19 @@ UMaterial* FMaterialUtilities::CreateMaterial(const FFlattenMaterial& InFlattenM
 		MergedSize = (bHasMetallic && MergedSize.X == 0) ? PropertySize : MergedSize;
 	}
 	bSameTextureSize &= bHasMetallic ? (SampleCount == InFlattenMaterial.GetPropertySamples(EFlattenMaterialProperties::Metallic).Num()) : true;
-	bSameTextureSize &= bHasRoughness ? (SampleCount == InFlattenMaterial.GetPropertySamples(EFlattenMaterialProperties::Roughness).Num()) : true;
 	bSameTextureSize &= bHasSpecular ? (SampleCount == InFlattenMaterial.GetPropertySamples(EFlattenMaterialProperties::Specular).Num()) : true;
+	bSameTextureSize &= bHasRoughness ? (SampleCount == InFlattenMaterial.GetPropertySamples(EFlattenMaterialProperties::Roughness).Num()) : true;
 
 	// Merge values into one texture if more than one material property exists
 	if (BakedMaterialPropertyCount > 1 && bSameTextureSize)
 	{
-		// Metallic = R, Roughness = G, Specular = B
+		// Metallic = R, Specular = G, Roughness = B
 		TArray<FColor> MergedSamples;
 		MergedSamples.AddZeroed(SampleCount);
 		
 		// R G B masks
-#if PLATFORM_LITTLE_ENDIAN
-		const uint32 ColorMask[3] = { 0x0000FF00, 0x00FF0000, 0xFF000000 };
-#else // PLATFORM_LITTLE_ENDIAN
 		const uint32 ColorMask[3] = { 0x00FF0000, 0x0000FF00, 0x000000FF };
-#endif
+
 		for (int32 PropertyIndex = 0; PropertyIndex < 3; ++PropertyIndex)
 		{
 			EFlattenMaterialProperties Property = (EFlattenMaterialProperties)(PropertyIndex + (int32)EFlattenMaterialProperties::Metallic);
@@ -1112,7 +1108,7 @@ UMaterial* FMaterialUtilities::CreateMaterial(const FFlattenMaterial& InFlattenM
 			}
 		}
 
-		const FString AssetName = TEXT("T_") + AssetBaseName + TEXT("_MRS");
+		const FString AssetName = TEXT("T_") + AssetBaseName + TEXT("_MSR");
 		const bool bSRGB = true;
 		UTexture2D* Texture = CreateTexture(InOuter, AssetBasePath / AssetName, MergedSize, MergedSamples, TC_Default, InTextureGroup, Flags, bSRGB);
 		OutGeneratedAssets.Add(Texture);
@@ -1135,28 +1131,28 @@ UMaterial* FMaterialUtilities::CreateMaterial(const FFlattenMaterial& InFlattenM
 			Material->Metallic.MaskA = 0;
 		}
 
-		// Roughness
-		if (bHasRoughness)
-		{
-			Material->Roughness.Expression = MergedExpression;
-			Material->Roughness.Mask = Material->Roughness.Expression->GetOutputs()[0].Mask;
-			Material->Roughness.MaskR = 0;
-			Material->Roughness.MaskG = 1;
-			Material->Roughness.MaskB = 0;
-			Material->Roughness.MaskA = 0;
-		}
-		
 		// Specular
 		if (bHasSpecular)
 		{
 			Material->Specular.Expression = MergedExpression;
 			Material->Specular.Mask = Material->Specular.Expression->GetOutputs()[0].Mask;
 			Material->Specular.MaskR = 0;
-			Material->Specular.MaskG = 0;
-			Material->Specular.MaskB = 1;
+			Material->Specular.MaskG = 1;
+			Material->Specular.MaskB = 0;
 			Material->Specular.MaskA = 0;
 		}
 
+		// Roughness
+		if (bHasRoughness)
+		{
+			Material->Roughness.Expression = MergedExpression;
+			Material->Roughness.Mask = Material->Roughness.Expression->GetOutputs()[0].Mask;
+			Material->Roughness.MaskR = 0;
+			Material->Roughness.MaskG = 0;
+			Material->Roughness.MaskB = 1;
+			Material->Roughness.MaskA = 0;
+		}
+
 		MaterialNodeY += MaterialNodeStepY;
 	}
 	else
@@ -1745,12 +1741,12 @@ void FMaterialUtilities::RemapUniqueMaterialIndices(const TArray<FSectionInfo>&
 				TArray<int32>& NewMeshMaterialMap = OutMaterialMap.Add(FMeshIdAndLOD(MeshIndex, LODIndex));
 				UStaticMesh* StaticMesh = InMeshData[MeshIndex].SourceStaticMesh;
 
-		if (!OutMeshShouldBakeVertexData[MeshIndex])
-		{
-			// No vertex data needed - could merge materials with other meshes.
-			// Set to 'nullptr' if don't need to bake vertex data to be able to merge materials with any meshes
-			// which don't require vertex data baking too.
-			StaticMesh = nullptr;
+				if (!OutMeshShouldBakeVertexData[MeshIndex])
+				{
+					// No vertex data needed - could merge materials with other meshes.
+					// Set to 'nullptr' if don't need to bake vertex data to be able to merge materials with any meshes
+					// which don't require vertex data baking too.
+					StaticMesh = nullptr;
 
 					for (int32 LocalMaterialIndex = 0; LocalMaterialIndex < MeshMaterialMap.Num(); LocalMaterialIndex++)
 					{
