* Fixed static analysis warning

* Fixed FSceneViewInitOptions::ShowOnlyPrimitives  (a TOptional) being assigned a value - if no "show only" primitives are provided, we must pass an unset TOptional - otherwise an empty set will mean no primitive should be visible.
** Solves an issue with FMaterialUtilities::ExportLandscapeMaterial() being called by the USDExporter & World Composition LOD level generation

diff --git a/Engine/Source/Developer/MaterialUtilities/Private/MaterialUtilities.cpp b/Engine/Source/Developer/MaterialUtilities/Private/MaterialUtilities.cpp
--- a/Engine/Source/Developer/MaterialUtilities/Private/MaterialUtilities.cpp
+++ b/Engine/Source/Developer/MaterialUtilities/Private/MaterialUtilities.cpp
@@ -98,6 +98,8 @@ UMaterialInterface* FMaterialUtilities::CreateProxyMaterialAndTextures(UPackage*
 
 	/** Create Proxy material and populate flags */
 	UMaterialInstanceConstant* Material = FMaterialUtilities::CreateInstancedMaterial(BaseMaterial, OuterPackage, AssetName, RF_Public | RF_Standalone);
+	check(Material);
+	
 	Material->BasePropertyOverrides.TwoSided = MaterialData.Material->IsTwoSided();
 	Material->BasePropertyOverrides.bOverride_TwoSided = MaterialData.Material->IsTwoSided();
 	Material->BasePropertyOverrides.DitheredLODTransition = MaterialData.Material->IsDitheredLODTransition();
@@ -807,11 +809,13 @@ static void RenderSceneToTexture(
 	FSceneViewInitOptions ViewInitOptions;
 	ViewInitOptions.SetViewRectangle(FIntRect(0, 0, TargetSize.X, TargetSize.Y));
 	ViewInitOptions.ViewFamily = &ViewFamily;
-	ViewInitOptions.ShowOnlyPrimitives = ShowOnlyPrimitives;
-	ViewInitOptions.HiddenPrimitives = HiddenPrimitives;
 	ViewInitOptions.ViewOrigin = ViewOrigin;
 	ViewInitOptions.ViewRotationMatrix = ViewRotationMatrix;
 	ViewInitOptions.ProjectionMatrix = ProjectionMatrix;
+
+	// If no "show only" primitives are provided, we must pass an unset TOptional - otherwise an empty set will mean no primitive should be visible.
+	ViewInitOptions.ShowOnlyPrimitives = !ShowOnlyPrimitives.IsEmpty() ? TOptional<TSet<FPrimitiveComponentId>>(ShowOnlyPrimitives) : TOptional<TSet<FPrimitiveComponentId>>();
+	ViewInitOptions.HiddenPrimitives = HiddenPrimitives;
 		
 	FSceneView* NewView = new FSceneView(ViewInitOptions);
 	NewView->CurrentBufferVisualizationMode = VisualizationMode;
