Fixed static analysis warning in MaterialUtilities.cpp

#rnx

diff --git a/Engine/Source/Developer/MaterialUtilities/Private/MaterialUtilities.cpp b/Engine/Source/Developer/MaterialUtilities/Private/MaterialUtilities.cpp
--- a/Engine/Source/Developer/MaterialUtilities/Private/MaterialUtilities.cpp
+++ b/Engine/Source/Developer/MaterialUtilities/Private/MaterialUtilities.cpp
@@ -137,14 +137,11 @@ UMaterialInterface* FMaterialUtilities::CreateProxyMaterialAndTextures(UPackage*
 			CreateParams.bSRGB = SRGBEnabledProperties.Contains(Property);
 
 			// Make sure the texture is a VT if required by the material sampler
-			if (Material != nullptr)
+			UTexture* DefaultTexture = nullptr;
+			Material->GetTextureParameterValue(ParameterInfo, DefaultTexture);
+			if (DefaultTexture)
 			{
-				UTexture* DefaultTexture = nullptr;
-				Material->GetTextureParameterValue(ParameterInfo, DefaultTexture);
-				if (DefaultTexture)
-				{
-					CreateParams.bVirtualTexture = DefaultTexture->VirtualTextureStreaming;
-				}
+				CreateParams.bVirtualTexture = DefaultTexture->VirtualTextureStreaming;
 			}
 
 			UTexture* Texture = FMaterialUtilities::CreateTexture(OuterPackage, TEXT("T_") + AssetName + TEXT("_") + TrimmedPropertyName, DataSize, ColorData, CreateParams, RF_Public | RF_Standalone);
