Fix to correctly check material property sizes in CreateFlattenMaterialInstance.

diff --git a/Engine/Source/Developer/MaterialUtilities/Private/ProxyMaterialUtilities.cpp b/Engine/Source/Developer/MaterialUtilities/Private/ProxyMaterialUtilities.cpp
--- a/Engine/Source/Developer/MaterialUtilities/Private/ProxyMaterialUtilities.cpp
+++ b/Engine/Source/Developer/MaterialUtilities/Private/ProxyMaterialUtilities.cpp
@@ -388,43 +388,44 @@ UMaterialInstanceConstant* FMaterialUtilities::CreateFlattenMaterialInstance(UPa
 	};
 
 	// Load textures and set switches accordingly
-	if (FlattenMaterial.GetPropertySize(EFlattenMaterialProperties::Diffuse).Num() > 0 && !(FlattenMaterial.IsPropertyConstant(EFlattenMaterialProperties::Diffuse) && FlattenMaterial.GetPropertySamples(EFlattenMaterialProperties::Diffuse)[0] == FColor::Black))
+	if (FlattenMaterial.ShouldGenerateDataForProperty(EFlattenMaterialProperties::Diffuse) && !(FlattenMaterial.IsPropertyConstant(EFlattenMaterialProperties::Diffuse) && FlattenMaterial.GetPropertySamples(EFlattenMaterialProperties::Diffuse)[0] == FColor::Black))
 	{
 		SetTextureParamConstVector(EFlattenMaterialProperties::Diffuse);
 	}
 
-	if (FlattenMaterial.GetPropertySize(EFlattenMaterialProperties::Normal).Num() > 1)
+	const FIntPoint NormalTextureSize = FlattenMaterial.GetPropertySize(EFlattenMaterialProperties::Normal);
+	if (NormalTextureSize.X * NormalTextureSize.Y > 1)
 	{
 		SetTextureParam(EFlattenMaterialProperties::Normal);
 	}
 
 	// Determine whether or not specific material properties are packed together into one texture (requires at least two to match (number of samples and texture size) in order to be packed
-	if (!bPackMetallic && (FlattenMaterial.GetPropertySize(EFlattenMaterialProperties::Metallic).Num() > 0 || !InMaterialProxySettings.bMetallicMap))
+	if (!bPackMetallic && (FlattenMaterial.ShouldGenerateDataForProperty(EFlattenMaterialProperties::Metallic) || !InMaterialProxySettings.bMetallicMap))
 	{
 		SetTextureParamConstScalar(EFlattenMaterialProperties::Metallic, InMaterialProxySettings.MetallicConstant);
 	}
 
-	if (!bPackRoughness && (FlattenMaterial.GetPropertySize(EFlattenMaterialProperties::Roughness).Num() > 0 || !InMaterialProxySettings.bRoughnessMap))
+	if (!bPackRoughness && (FlattenMaterial.ShouldGenerateDataForProperty(EFlattenMaterialProperties::Roughness) || !InMaterialProxySettings.bRoughnessMap))
 	{
 		SetTextureParamConstScalar(EFlattenMaterialProperties::Roughness, InMaterialProxySettings.RoughnessConstant);
 	}
 
-	if (!bPackSpecular && (FlattenMaterial.GetPropertySize(EFlattenMaterialProperties::Specular).Num() > 0 || !InMaterialProxySettings.bSpecularMap))
+	if (!bPackSpecular && (FlattenMaterial.ShouldGenerateDataForProperty(EFlattenMaterialProperties::Specular) || !InMaterialProxySettings.bSpecularMap))
 	{
 		SetTextureParamConstScalar(EFlattenMaterialProperties::Specular, InMaterialProxySettings.SpecularConstant);
 	}
 
-	if (FlattenMaterial.GetPropertySize(EFlattenMaterialProperties::Opacity).Num() > 0 || !InMaterialProxySettings.bOpacityMap)
+	if (FlattenMaterial.ShouldGenerateDataForProperty(EFlattenMaterialProperties::Opacity) || !InMaterialProxySettings.bOpacityMap)
 	{
 		SetTextureParamConstScalar(EFlattenMaterialProperties::Opacity, InMaterialProxySettings.OpacityConstant);
 	}
 
-	if (FlattenMaterial.GetPropertySize(EFlattenMaterialProperties::OpacityMask).Num() > 0 || !InMaterialProxySettings.bOpacityMaskMap)
+	if (FlattenMaterial.ShouldGenerateDataForProperty(EFlattenMaterialProperties::OpacityMask) || !InMaterialProxySettings.bOpacityMaskMap)
 	{
 		SetTextureParamConstScalar(EFlattenMaterialProperties::OpacityMask, InMaterialProxySettings.OpacityMaskConstant);
 	}
 
-	if (FlattenMaterial.GetPropertySize(EFlattenMaterialProperties::AmbientOcclusion).Num() > 0 || !InMaterialProxySettings.bAmbientOcclusionMap)
+	if (FlattenMaterial.ShouldGenerateDataForProperty(EFlattenMaterialProperties::AmbientOcclusion) || !InMaterialProxySettings.bAmbientOcclusionMap)
 	{
 		SetTextureParamConstScalar(EFlattenMaterialProperties::AmbientOcclusion, InMaterialProxySettings.AmbientOcclusionConstant);
 	}
@@ -491,7 +492,8 @@ UMaterialInstanceConstant* FMaterialUtilities::CreateFlattenMaterialInstance(UPa
 	}
 
 	// Emissive is a special case due to the scaling variable
-	if (FlattenMaterial.GetPropertySamples(EFlattenMaterialProperties::Emissive).Num() > 0 && !(FlattenMaterial.GetPropertySize(EFlattenMaterialProperties::Emissive).Num() == 1 && FlattenMaterial.GetPropertySamples(EFlattenMaterialProperties::Emissive)[0] == FColor::Black))
+	const FIntPoint EmissiveTextureSize = FlattenMaterial.GetPropertySize(EFlattenMaterialProperties::Emissive);
+	if (FlattenMaterial.DoesPropertyContainData(EFlattenMaterialProperties::Emissive) && !(EmissiveTextureSize.X * EmissiveTextureSize.Y == 1 && FlattenMaterial.GetPropertySamples(EFlattenMaterialProperties::Emissive)[0] == FColor::Black))
 	{
 		SetTextureParamConstVector(EFlattenMaterialProperties::Emissive);
 
