Add support for masked materials with Merge Actors flatten material.
Fix condition for bRenderOpacityMask to be correctly based on the OpacityMask property instead of Opacity property.

diff --git a/Engine/Plugins/Experimental/ProxyLODPlugin/Source/ProxyLOD/Private/ProxyLODMaterialTransferUtilities.cpp b/Engine/Plugins/Experimental/ProxyLODPlugin/Source/ProxyLOD/Private/ProxyLODMaterialTransferUtilities.cpp
--- a/Engine/Plugins/Experimental/ProxyLODPlugin/Source/ProxyLOD/Private/ProxyLODMaterialTransferUtilities.cpp
+++ b/Engine/Plugins/Experimental/ProxyLODPlugin/Source/ProxyLOD/Private/ProxyLODMaterialTransferUtilities.cpp
@@ -1680,53 +1680,51 @@ namespace
 		}
 	}
 
-
-
-
- // Map Diffuse, Specular, Metallic, Roughness, Normal, Emissive, Opacity to the correct materials.
- // NB: EFlattenMaterialProperties: SubSurface, OpacityMask and AmbientOcclusion are not included
-void MapFlattenMaterial(const EFlattenMaterialProperties PropertyType, 
-	const FMeshDescription& DstRawMesh,
-	const FMeshDescriptionArrayAdapter& SrcMeshAdapter,
-	const ProxyLOD::FSrcDataGrid& SuperSampledCorrespondenceGrid, 
-	const ProxyLOD::FRasterGrid& SuperSampledDstUVGrid,
-	const ProxyLOD::FRasterGrid& DstUVGrid, 
-	const TArray<FFlattenMaterial>& InputMaterials, 
-	const FColor UnresolvedSrcColor,
-	FFlattenMaterial& OutMaterial)
-{
-	switch (PropertyType)
+	// Map Diffuse, Specular, Metallic, Roughness, Normal, Emissive, Opacity to the correct materials.
+	// NB: EFlattenMaterialProperties: SubSurface and AmbientOcclusion are not included
+	void MapFlattenMaterial(const EFlattenMaterialProperties PropertyType,
+		const FMeshDescription& DstRawMesh,
+		const FMeshDescriptionArrayAdapter& SrcMeshAdapter,
+		const ProxyLOD::FSrcDataGrid& SuperSampledCorrespondenceGrid,
+		const ProxyLOD::FRasterGrid& SuperSampledDstUVGrid,
+		const ProxyLOD::FRasterGrid& DstUVGrid,
+		const TArray<FFlattenMaterial>& InputMaterials,
+		const FColor UnresolvedSrcColor,
+		FFlattenMaterial& OutMaterial)
 	{
-	case EFlattenMaterialProperties::Diffuse:
-		TMapFlattenMaterial<EFlattenMaterialProperties::Diffuse>(DstRawMesh, SrcMeshAdapter, SuperSampledCorrespondenceGrid, SuperSampledDstUVGrid, DstUVGrid, InputMaterials, OutMaterial, UnresolvedSrcColor);
-		break;
-	case EFlattenMaterialProperties::Specular:
-		TMapFlattenMaterial<EFlattenMaterialProperties::Specular>(DstRawMesh, SrcMeshAdapter, SuperSampledCorrespondenceGrid, SuperSampledDstUVGrid, DstUVGrid, InputMaterials, OutMaterial);
-		break;
-	case EFlattenMaterialProperties::Metallic:
-		TMapFlattenMaterial<EFlattenMaterialProperties::Metallic>(DstRawMesh, SrcMeshAdapter, SuperSampledCorrespondenceGrid, SuperSampledDstUVGrid, DstUVGrid, InputMaterials, OutMaterial);
-		break;
-	case EFlattenMaterialProperties::Roughness:
-		TMapFlattenMaterial<EFlattenMaterialProperties::Roughness>(DstRawMesh, SrcMeshAdapter, SuperSampledCorrespondenceGrid, SuperSampledDstUVGrid, DstUVGrid, InputMaterials, OutMaterial);
-		break;
-	case EFlattenMaterialProperties::Normal:
-		TMapFlattenMaterial<EFlattenMaterialProperties::Normal>(DstRawMesh, SrcMeshAdapter, SuperSampledCorrespondenceGrid, SuperSampledDstUVGrid, DstUVGrid, InputMaterials, OutMaterial);
-		break;
-	case EFlattenMaterialProperties::Emissive:
-		TMapFlattenMaterial<EFlattenMaterialProperties::Emissive>(DstRawMesh, SrcMeshAdapter, SuperSampledCorrespondenceGrid, SuperSampledDstUVGrid, DstUVGrid, InputMaterials, OutMaterial);
-		break;
-	case EFlattenMaterialProperties::Opacity:
-		TMapFlattenMaterial<EFlattenMaterialProperties::Opacity>(DstRawMesh, SrcMeshAdapter, SuperSampledCorrespondenceGrid, SuperSampledDstUVGrid, DstUVGrid, InputMaterials, OutMaterial);
-		break;
-	default:
-		return;
-		// do nothing.  
-	};
-}
-
+		switch (PropertyType)
+		{
+		case EFlattenMaterialProperties::Diffuse:
+			TMapFlattenMaterial<EFlattenMaterialProperties::Diffuse>(DstRawMesh, SrcMeshAdapter, SuperSampledCorrespondenceGrid, SuperSampledDstUVGrid, DstUVGrid, InputMaterials, OutMaterial, UnresolvedSrcColor);
+			break;
+		case EFlattenMaterialProperties::Specular:
+			TMapFlattenMaterial<EFlattenMaterialProperties::Specular>(DstRawMesh, SrcMeshAdapter, SuperSampledCorrespondenceGrid, SuperSampledDstUVGrid, DstUVGrid, InputMaterials, OutMaterial);
+			break;
+		case EFlattenMaterialProperties::Metallic:
+			TMapFlattenMaterial<EFlattenMaterialProperties::Metallic>(DstRawMesh, SrcMeshAdapter, SuperSampledCorrespondenceGrid, SuperSampledDstUVGrid, DstUVGrid, InputMaterials, OutMaterial);
+			break;
+		case EFlattenMaterialProperties::Roughness:
+			TMapFlattenMaterial<EFlattenMaterialProperties::Roughness>(DstRawMesh, SrcMeshAdapter, SuperSampledCorrespondenceGrid, SuperSampledDstUVGrid, DstUVGrid, InputMaterials, OutMaterial);
+			break;
+		case EFlattenMaterialProperties::Normal:
+			TMapFlattenMaterial<EFlattenMaterialProperties::Normal>(DstRawMesh, SrcMeshAdapter, SuperSampledCorrespondenceGrid, SuperSampledDstUVGrid, DstUVGrid, InputMaterials, OutMaterial);
+			break;
+		case EFlattenMaterialProperties::Emissive:
+			TMapFlattenMaterial<EFlattenMaterialProperties::Emissive>(DstRawMesh, SrcMeshAdapter, SuperSampledCorrespondenceGrid, SuperSampledDstUVGrid, DstUVGrid, InputMaterials, OutMaterial);
+			break;
+		case EFlattenMaterialProperties::Opacity:
+			TMapFlattenMaterial<EFlattenMaterialProperties::Opacity>(DstRawMesh, SrcMeshAdapter, SuperSampledCorrespondenceGrid, SuperSampledDstUVGrid, DstUVGrid, InputMaterials, OutMaterial);
+			break;
+		case EFlattenMaterialProperties::OpacityMask:
+			TMapFlattenMaterial<EFlattenMaterialProperties::OpacityMask>(DstRawMesh, SrcMeshAdapter, SuperSampledCorrespondenceGrid, SuperSampledDstUVGrid, DstUVGrid, InputMaterials, OutMaterial);
+			break;
+		default:
+			return;
+			// do nothing.  
+		};
+	}
 }// end anonymous namespace
 
-
 void ProxyLOD::MapFlattenMaterials(const FMeshDescription& DstRawMesh, const FMeshDescriptionArrayAdapter& SrcMeshAdapter,
 	const ProxyLOD::FSrcDataGrid& SuperSampledCorrespondenceGrid,
 	const ProxyLOD::FRasterGrid& SuperSampledDstUVGrid,
diff --git a/Engine/Source/Developer/MaterialUtilities/Private/MaterialUtilities.cpp b/Engine/Source/Developer/MaterialUtilities/Private/MaterialUtilities.cpp
--- a/Engine/Source/Developer/MaterialUtilities/Private/MaterialUtilities.cpp
+++ b/Engine/Source/Developer/MaterialUtilities/Private/MaterialUtilities.cpp
@@ -1354,6 +1354,37 @@ UMaterial* FMaterialUtilities::CreateMaterial(const FFlattenMaterial& InFlattenM
 		MaterialNodeY += MaterialNodeStepY;
 	}
 
+	if (InFlattenMaterial.IsPropertyConstant(EFlattenMaterialProperties::OpacityMask))
+	{
+		// Set OpacityMask to constant
+		FLinearColor OpacityMask = FLinearColor(InFlattenMaterial.GetPropertySamples(EFlattenMaterialProperties::OpacityMask)[0]);
+		auto OpacityMaskExpression = NewObject<UMaterialExpressionConstant>(Material);
+		OpacityMaskExpression->R = OpacityMask.R;
+		OpacityMaskExpression->MaterialExpressionEditorX = -400;
+		OpacityMaskExpression->MaterialExpressionEditorY = MaterialNodeY;
+		Material->GetExpressionCollection().AddExpression(OpacityMaskExpression);
+		MaterialEditorOnly->OpacityMask.Expression = OpacityMaskExpression;
+
+		MaterialNodeY += MaterialNodeStepY;
+	}
+	else if (InFlattenMaterial.DoesPropertyContainData(EFlattenMaterialProperties::OpacityMask))
+	{
+		const FString AssetName = TEXT("T_") + AssetBaseName + TEXT("_OM");
+		const bool bSRGB = true;
+		UTexture2D* Texture = CreateTexture(InOuter, AssetBasePath / AssetName, InFlattenMaterial.GetPropertySize(EFlattenMaterialProperties::OpacityMask), InFlattenMaterial.GetPropertySamples(EFlattenMaterialProperties::OpacityMask), TC_Default, InTextureGroup, Flags, bSRGB);
+		OutGeneratedAssets.Add(Texture);
+
+		//Assign opacity to the material
+		UMaterialExpressionTextureSample* OpacityMaskExpression = NewObject<UMaterialExpressionTextureSample>(Material);
+		OpacityMaskExpression->Texture = Texture;
+		OpacityMaskExpression->SamplerType = EMaterialSamplerType::SAMPLERTYPE_Color;
+		OpacityMaskExpression->MaterialExpressionEditorX = -400;
+		OpacityMaskExpression->MaterialExpressionEditorY = MaterialNodeY;
+		Material->GetExpressionCollection().AddExpression(OpacityMaskExpression);
+		MaterialEditorOnly->OpacityMask.Expression = OpacityMaskExpression;
+		MaterialNodeY += MaterialNodeStepY;
+	}
+
 	if (InFlattenMaterial.IsPropertyConstant(EFlattenMaterialProperties::SubSurface))
 	{
 		// Set Emissive to constant
@@ -1524,7 +1555,8 @@ FFlattenMaterial FMaterialUtilities::CreateFlattenMaterialWithSettings(const FMa
 		MaximumSize = (InMaterialLODSettings.bNormalMap && (MaximumSize.X < InMaterialLODSettings.NormalTextureSize.X)) ? InMaterialLODSettings.NormalTextureSize :			MaximumSize;
 		MaximumSize = (InMaterialLODSettings.bEmissiveMap && (MaximumSize.X < InMaterialLODSettings.EmissiveTextureSize.X)) ? InMaterialLODSettings.EmissiveTextureSize :	MaximumSize;
 		MaximumSize = (InMaterialLODSettings.bOpacityMap && (MaximumSize.X < InMaterialLODSettings.OpacityTextureSize.X)) ? InMaterialLODSettings.OpacityTextureSize :		MaximumSize;
-	
+		MaximumSize = (InMaterialLODSettings.bOpacityMaskMap && (MaximumSize.X < InMaterialLODSettings.OpacityMaskTextureSize.X)) ? InMaterialLODSettings.OpacityMaskTextureSize : MaximumSize;
+
 		Material.RenderSize = MaximumSize;
 		Material.SetPropertySize(EFlattenMaterialProperties::Diffuse, InMaterialLODSettings.DiffuseTextureSize);
 		Material.SetPropertySize(EFlattenMaterialProperties::Specular, InMaterialLODSettings.bSpecularMap ? InMaterialLODSettings.SpecularTextureSize : FIntPoint::ZeroValue);
@@ -1533,6 +1565,7 @@ FFlattenMaterial FMaterialUtilities::CreateFlattenMaterialWithSettings(const FMa
 		Material.SetPropertySize(EFlattenMaterialProperties::Normal, InMaterialLODSettings.bNormalMap ? InMaterialLODSettings.NormalTextureSize : FIntPoint::ZeroValue);
 		Material.SetPropertySize(EFlattenMaterialProperties::Emissive, InMaterialLODSettings.bEmissiveMap ? InMaterialLODSettings.EmissiveTextureSize : FIntPoint::ZeroValue);
 		Material.SetPropertySize(EFlattenMaterialProperties::Opacity, InMaterialLODSettings.bOpacityMap ? InMaterialLODSettings.OpacityTextureSize : FIntPoint::ZeroValue);
+		Material.SetPropertySize(EFlattenMaterialProperties::OpacityMask, InMaterialLODSettings.bOpacityMaskMap ? InMaterialLODSettings.OpacityMaskTextureSize : FIntPoint::ZeroValue);
 	}
 	else if (InMaterialLODSettings.TextureSizingType == TextureSizingType_UseAutomaticBiasedSizes)
 	{
@@ -1550,6 +1583,7 @@ FFlattenMaterial FMaterialUtilities::CreateFlattenMaterialWithSettings(const FMa
 		Material.SetPropertySize(EFlattenMaterialProperties::Roughness, (InMaterialLODSettings.bRoughnessMap) ? PropertiesSize : FIntPoint::ZeroValue );
 		Material.SetPropertySize(EFlattenMaterialProperties::Emissive, (InMaterialLODSettings.bEmissiveMap) ? PropertiesSize : FIntPoint::ZeroValue );
 		Material.SetPropertySize(EFlattenMaterialProperties::Opacity, (InMaterialLODSettings.bOpacityMap) ? PropertiesSize : FIntPoint::ZeroValue );
+		Material.SetPropertySize(EFlattenMaterialProperties::OpacityMask, (InMaterialLODSettings.bOpacityMaskMap) ? PropertiesSize : FIntPoint::ZeroValue);
 	}
 	else if (InMaterialLODSettings.TextureSizingType == TextureSizingType_UseSingleTextureSize)
 	{
@@ -1561,6 +1595,7 @@ FFlattenMaterial FMaterialUtilities::CreateFlattenMaterialWithSettings(const FMa
 		Material.SetPropertySize(EFlattenMaterialProperties::Normal, (InMaterialLODSettings.bNormalMap) ? InMaterialLODSettings.TextureSize : FIntPoint::ZeroValue);
 		Material.SetPropertySize(EFlattenMaterialProperties::Emissive, (InMaterialLODSettings.bEmissiveMap) ? InMaterialLODSettings.TextureSize : FIntPoint::ZeroValue);
 		Material.SetPropertySize(EFlattenMaterialProperties::Opacity, (InMaterialLODSettings.bOpacityMap) ? InMaterialLODSettings.TextureSize : FIntPoint::ZeroValue);
+		Material.SetPropertySize(EFlattenMaterialProperties::OpacityMask, (InMaterialLODSettings.bOpacityMaskMap) ? InMaterialLODSettings.TextureSize : FIntPoint::ZeroValue);
 	}
 	else
 	{
@@ -1873,6 +1908,15 @@ void FMaterialUtilities::ResizeFlattenMaterial(FFlattenMaterial& InFlattenMateri
 			InFlattenMaterial.GetPropertySamples(EFlattenMaterialProperties::Opacity).Append(NewSamples);
 			InFlattenMaterial.SetPropertySize(EFlattenMaterialProperties::Opacity, FIntPoint(PropertiesSizeX, PropertiesSizeX));
 		}
+
+		if (InFlattenMaterial.GetPropertySamples(EFlattenMaterialProperties::OpacityMask).Num() && MaterialSettings.bOpacityMaskMap && InFlattenMaterial.GetPropertySize(EFlattenMaterialProperties::OpacityMask).X != PropertiesSizeX)
+		{
+			TArray<FColor> NewSamples;
+			FImageUtils::ImageResize(InFlattenMaterial.GetPropertySize(EFlattenMaterialProperties::OpacityMask).X, InFlattenMaterial.GetPropertySize(EFlattenMaterialProperties::OpacityMask).Y, InFlattenMaterial.GetPropertySamples(EFlattenMaterialProperties::OpacityMask), PropertiesSizeX, PropertiesSizeX, NewSamples, false);
+			InFlattenMaterial.GetPropertySamples(EFlattenMaterialProperties::OpacityMask).Reset(NewSamples.Num());
+			InFlattenMaterial.GetPropertySamples(EFlattenMaterialProperties::OpacityMask).Append(NewSamples);
+			InFlattenMaterial.SetPropertySize(EFlattenMaterialProperties::OpacityMask, FIntPoint(PropertiesSizeX, PropertiesSizeX));
+		}
 	}
 	else if (MaterialSettings.TextureSizingType == TextureSizingType_UseManualOverrideTextureSize)
 	{
@@ -1938,6 +1982,15 @@ void FMaterialUtilities::ResizeFlattenMaterial(FFlattenMaterial& InFlattenMateri
 			InFlattenMaterial.GetPropertySamples(EFlattenMaterialProperties::Opacity).Append(NewSamples);
 			InFlattenMaterial.SetPropertySize(EFlattenMaterialProperties::Opacity, MaterialSettings.OpacityTextureSize);
 		}
+
+		if (InFlattenMaterial.GetPropertySamples(EFlattenMaterialProperties::OpacityMask).Num() && MaterialSettings.bOpacityMaskMap && InFlattenMaterial.GetPropertySize(EFlattenMaterialProperties::OpacityMask) != MaterialSettings.OpacityMaskTextureSize)
+		{
+			TArray<FColor> NewSamples;
+			FImageUtils::ImageResize(InFlattenMaterial.GetPropertySize(EFlattenMaterialProperties::OpacityMask).X, InFlattenMaterial.GetPropertySize(EFlattenMaterialProperties::OpacityMask).Y, InFlattenMaterial.GetPropertySamples(EFlattenMaterialProperties::OpacityMask), MaterialSettings.OpacityMaskTextureSize.X, MaterialSettings.OpacityMaskTextureSize.Y, NewSamples, false);
+			InFlattenMaterial.GetPropertySamples(EFlattenMaterialProperties::OpacityMask).Reset(NewSamples.Num());
+			InFlattenMaterial.GetPropertySamples(EFlattenMaterialProperties::OpacityMask).Append(NewSamples);
+			InFlattenMaterial.SetPropertySize(EFlattenMaterialProperties::OpacityMask, MaterialSettings.OpacityMaskTextureSize);
+		}
 	}
 }
 
@@ -2287,7 +2340,7 @@ bool FMaterialUtilities::ExportMaterial(struct FMaterialMergeData& InMaterialDat
 	const bool bRenderNormal = (Material->GetMaterial()->HasNormalConnected() || Material->GetMaterial()->bUseMaterialAttributes) && OutFlattenMaterial.ShouldGenerateDataForProperty(EFlattenMaterialProperties::Normal);
 	const bool bRenderTangent = (Material->GetMaterial()->IsPropertyConnected(MP_Tangent) || Material->GetMaterial()->bUseMaterialAttributes) && OutFlattenMaterial.ShouldGenerateDataForProperty(EFlattenMaterialProperties::Tangent);
 	const bool bRenderEmissive = (Material->GetMaterial()->IsPropertyConnected(MP_EmissiveColor) || Material->GetMaterial()->bUseMaterialAttributes) && OutFlattenMaterial.ShouldGenerateDataForProperty(EFlattenMaterialProperties::Emissive);
-	const bool bRenderOpacityMask = Material->IsPropertyActive(MP_OpacityMask) && Material->GetBlendMode() == BLEND_Masked && OutFlattenMaterial.ShouldGenerateDataForProperty(EFlattenMaterialProperties::Opacity);
+	const bool bRenderOpacityMask = Material->IsPropertyActive(MP_OpacityMask) && Material->GetBlendMode() == BLEND_Masked && OutFlattenMaterial.ShouldGenerateDataForProperty(EFlattenMaterialProperties::OpacityMask);
 	const bool bRenderOpacity = Material->IsPropertyActive(MP_Opacity) && IsTranslucentBlendMode(Material->GetBlendMode()) && OutFlattenMaterial.ShouldGenerateDataForProperty(EFlattenMaterialProperties::Opacity);
 	const bool bRenderSubSurface = Material->IsPropertyActive(MP_SubsurfaceColor) && OutFlattenMaterial.ShouldGenerateDataForProperty(EFlattenMaterialProperties::SubSurface);
 	const bool bRenderMetallic = Material->IsPropertyActive(MP_Metallic) && OutFlattenMaterial.ShouldGenerateDataForProperty(EFlattenMaterialProperties::Metallic);
@@ -2382,6 +2435,7 @@ bool FMaterialUtilities::ExportMaterial(struct FMaterialMergeData& InMaterialDat
 		PRAGMA_ENABLE_DEPRECATION_WARNINGS
 		OutFlattenMaterial.SetPropertySize(EFlattenMaterialProperties::OpacityMask, Size);
 	}
+
 	if (bRenderOpacity)
 	{
 		Size = OutFlattenMaterial.GetPropertySize(EFlattenMaterialProperties::Opacity);
@@ -2391,6 +2445,7 @@ bool FMaterialUtilities::ExportMaterial(struct FMaterialMergeData& InMaterialDat
 		PRAGMA_ENABLE_DEPRECATION_WARNINGS
 		OutFlattenMaterial.SetPropertySize(EFlattenMaterialProperties::Opacity, Size);
 	}
+
 	if (bRenderEmissive)
 	{
 		Size = OutFlattenMaterial.GetPropertySize(EFlattenMaterialProperties::Emissive);
