Fix even more Strata blend mode.

diff --git a/Engine/Plugins/Experimental/Water/Source/Runtime/Private/WaterMeshSceneProxy.cpp b/Engine/Plugins/Experimental/Water/Source/Runtime/Private/WaterMeshSceneProxy.cpp
--- a/Engine/Plugins/Experimental/Water/Source/Runtime/Private/WaterMeshSceneProxy.cpp
+++ b/Engine/Plugins/Experimental/Water/Source/Runtime/Private/WaterMeshSceneProxy.cpp
@@ -335,7 +335,7 @@ void FWaterMeshSceneProxy::GetDynamicMeshElements(const TArray<const FSceneView*
 					if (const FMaterial* BucketMaterial = MaterialRenderProxy->GetMaterialNoFallback(GetScene().GetFeatureLevel()))
 					{
 						// Preemptively turn off depth rendering for this mesh batch if the material doesn't need it
-						bUseForDepthPass = !BucketMaterial->GetShadingModels().HasShadingModel(MSM_SingleLayerWater) && BucketMaterial->GetBlendMode() != EBlendMode::BLEND_Translucent;
+						bUseForDepthPass = !BucketMaterial->GetShadingModels().HasShadingModel(MSM_SingleLayerWater) && !IsTranslucentOnlyBlendMode(*BucketMaterial);
 					}
 
 					bMaterialDrawn = true;
diff --git a/Engine/Plugins/Runtime/nDisplay/Source/DisplayClusterLightCardEditorShaders/Private/DisplayClusterMeshProjectionRenderer.cpp b/Engine/Plugins/Runtime/nDisplay/Source/DisplayClusterLightCardEditorShaders/Private/DisplayClusterMeshProjectionRenderer.cpp
--- a/Engine/Plugins/Runtime/nDisplay/Source/DisplayClusterLightCardEditorShaders/Private/DisplayClusterMeshProjectionRenderer.cpp
+++ b/Engine/Plugins/Runtime/nDisplay/Source/DisplayClusterLightCardEditorShaders/Private/DisplayClusterMeshProjectionRenderer.cpp
@@ -372,7 +372,7 @@ protected:
 		{
 			SortKey.BasePass.VertexShaderHash = (VertexShader ? VertexShader->GetSortKey() : 0) & 0xFFFF;
 			SortKey.BasePass.PixelShaderHash = PixelShader ? PixelShader->GetSortKey() : 0;
-			SortKey.BasePass.Masked = Material.GetBlendMode() == EBlendMode::BLEND_Masked ? 1 : 0;
+			SortKey.BasePass.Masked = IsMaskedBlendMode(Material) ? 1 : 0;
 		}
 
 		return SortKey;
diff --git a/Engine/Source/Developer/MaterialUtilities/Private/MaterialUtilities.cpp b/Engine/Source/Developer/MaterialUtilities/Private/MaterialUtilities.cpp
--- a/Engine/Source/Developer/MaterialUtilities/Private/MaterialUtilities.cpp
+++ b/Engine/Source/Developer/MaterialUtilities/Private/MaterialUtilities.cpp
@@ -558,7 +558,6 @@ public:
 			const bool bIsOpaqueOrMasked = IsOpaqueOrMaskedBlendMode(*MaterialInterface);
 			UMaterial* ProxyMaterial = MaterialInterface->GetMaterial();
 			check(ProxyMaterial);
-			EBlendMode BlendMode = MaterialInterface->GetBlendMode();
 			FExportMaterialCompiler ProxyCompiler(Compiler);
 			const uint32 ForceCast_Exact_Replicate = MFCF_ForceCast | MFCF_ExactMatch | MFCF_ReplicateValue;
 									
@@ -726,30 +725,26 @@ public:
 		return Ar << V.MaterialInterface;
 	}
 
-	static bool WillFillData(EBlendMode InBlendMode, EMaterialProperty InMaterialProperty)
+	static bool WillFillData(bool bIsOpaque, EMaterialProperty InMaterialProperty)
 	{
 		if (InMaterialProperty == MP_EmissiveColor)
 		{
 			return true;
 		}
 
-		switch (InBlendMode)
+		if (bIsOpaque)
 		{
-		case BLEND_Opaque:
+			switch (InMaterialProperty)
 			{
-				switch (InMaterialProperty)
-				{
-				case MP_BaseColor:			return true;
-				case MP_Specular:			return true;
-				case MP_Normal:				return true;
-				case MP_Tangent:			return true;
-				case MP_Metallic:			return true;
-				case MP_Roughness:			return true;
-				case MP_Anisotropy:			return true;
-				case MP_AmbientOcclusion:	return true;
-				}
+			case MP_BaseColor:			return true;
+			case MP_Specular:			return true;
+			case MP_Normal:				return true;
+			case MP_Tangent:			return true;
+			case MP_Metallic:			return true;
+			case MP_Roughness:			return true;
+			case MP_Anisotropy:			return true;
+			case MP_AmbientOcclusion:	return true;
 			}
-			break;
 		}
 		return false;
 	}
@@ -960,9 +955,9 @@ FIntPoint FMaterialUtilities::FindMaxTextureSize(UMaterialInterface* InMaterialI
 	return MaxSize;
 }
 
-bool FMaterialUtilities::SupportsExport(EBlendMode InBlendMode, EMaterialProperty InMaterialProperty)
+bool FMaterialUtilities::SupportsExport(bool bIsOpaque, EMaterialProperty InMaterialProperty)
 {
-	return FExportMaterialProxy::WillFillData(InBlendMode, InMaterialProperty);
+	return FExportMaterialProxy::WillFillData(bIsOpaque, InMaterialProperty);
 }
 
 static bool ExportLandscapeMaterial(const ALandscapeProxy* InLandscape, const TSet<FPrimitiveComponentId>& ShowOnlyPrimitives, const TSet<FPrimitiveComponentId>& HiddenPrimitives, FFlattenMaterial& OutFlattenMaterial)
diff --git a/Engine/Source/Developer/MaterialUtilities/Public/MaterialUtilities.h b/Engine/Source/Developer/MaterialUtilities/Public/MaterialUtilities.h
--- a/Engine/Source/Developer/MaterialUtilities/Public/MaterialUtilities.h
+++ b/Engine/Source/Developer/MaterialUtilities/Public/MaterialUtilities.h
@@ -272,7 +272,7 @@ public:
 	/**
 	 * Whether material utilities support exporting specified material blend mode and property 
 	 */
-	static bool SupportsExport(EBlendMode InBlendMode, EMaterialProperty InMaterialProperty);
+	static bool SupportsExport(bool bIsOpaque, EMaterialProperty InMaterialProperty);
 
 	/**
 	 * Flattens specified landscape material
diff --git a/Engine/Source/Developer/MeshUtilities/Private/MeshRepresentationCommon.cpp b/Engine/Source/Developer/MeshUtilities/Private/MeshRepresentationCommon.cpp
--- a/Engine/Source/Developer/MeshUtilities/Private/MeshRepresentationCommon.cpp
+++ b/Engine/Source/Developer/MeshUtilities/Private/MeshRepresentationCommon.cpp
@@ -215,7 +215,7 @@ void MeshRepresentation::SetupEmbreeScene(
 					{
 						if (MaterialBlendModes.IsValidIndex(Section.MaterialIndex))
 						{
-							bTriangleIsOpaqueOrMasked = !IsTranslucentBlendMode(MaterialBlendModes[Section.MaterialIndex].BlendMode) && MaterialBlendModes[Section.MaterialIndex].bAffectDistanceFieldLighting; // STRATA_TODO_BLENDMODE
+							bTriangleIsOpaqueOrMasked = !IsTranslucentBlendMode(MaterialBlendModes[Section.MaterialIndex].BlendMode, MaterialBlendModes[Section.MaterialIndex].StrataBlendMode) && MaterialBlendModes[Section.MaterialIndex].bAffectDistanceFieldLighting;
 						}
 
 						break;
diff --git a/Engine/Source/Developer/MeshUtilities/Public/MeshUtilities.h b/Engine/Source/Developer/MeshUtilities/Public/MeshUtilities.h
--- a/Engine/Source/Developer/MeshUtilities/Public/MeshUtilities.h
+++ b/Engine/Source/Developer/MeshUtilities/Public/MeshUtilities.h
@@ -52,6 +52,7 @@ class FSignedDistanceFieldBuildMaterialData
 {
 public:
 	EBlendMode BlendMode = BLEND_Opaque;
+	EStrataBlendMode StrataBlendMode = SBM_Opaque;
 	bool bTwoSided = false;
 	bool bAffectDistanceFieldLighting = true;
 };
diff --git a/Engine/Source/Editor/MaterialEditor/Private/SMaterialEditorViewport.cpp b/Engine/Source/Editor/MaterialEditor/Private/SMaterialEditorViewport.cpp
--- a/Engine/Source/Editor/MaterialEditor/Private/SMaterialEditorViewport.cpp
+++ b/Engine/Source/Editor/MaterialEditor/Private/SMaterialEditorViewport.cpp
@@ -181,7 +181,7 @@ FLinearColor FMaterialEditorViewportClient::GetBackgroundColor() const
 				{
 					BackgroundColor = FLinearColor::White;
 				}
-				else if (IsTranslucentOnlyBlendMode(*MaterialInterface) || PreviewBlendMode == BLEND_AlphaComposite || IsAlphaHoldoutBlendMode(*MaterialInterface))
+				else if (IsTranslucentOnlyBlendMode(*MaterialInterface) || IsAlphaCompositeBlendMode(*MaterialInterface) || IsAlphaHoldoutBlendMode(*MaterialInterface))
 				{
 					BackgroundColor = FColor(64, 64, 64);
 				}
diff --git a/Engine/Source/Editor/UnrealEd/Private/EditorExporters.cpp b/Engine/Source/Editor/UnrealEd/Private/EditorExporters.cpp
--- a/Engine/Source/Editor/UnrealEd/Private/EditorExporters.cpp
+++ b/Engine/Source/Editor/UnrealEd/Private/EditorExporters.cpp
@@ -1066,9 +1066,7 @@ static void ExportMaterialPropertyTexture(const FString& BMPFilename, UMaterialI
 	TArray<FColor> OutputBMP;
 	FIntPoint OutSize{};
 
-	const TEnumAsByte<EBlendMode> BlendMode = Material->GetBlendMode();
-
-	bool bIsValidProperty = FMaterialUtilities::SupportsExport(BlendMode, MatProp);
+	bool bIsValidProperty = FMaterialUtilities::SupportsExport(IsOpaqueBlendMode(*Material), MatProp);
 	if (bIsValidProperty)
 	{
 		FMeshData MeshSettings;
diff --git a/Engine/Source/Editor/UnrealEd/Private/Lightmass/LightmassRender.cpp b/Engine/Source/Editor/UnrealEd/Private/Lightmass/LightmassRender.cpp
--- a/Engine/Source/Editor/UnrealEd/Private/Lightmass/LightmassRender.cpp
+++ b/Engine/Source/Editor/UnrealEd/Private/Lightmass/LightmassRender.cpp
@@ -396,7 +396,6 @@ public:
 		if( Property == MP_EmissiveColor )
 		{
 			UMaterial* ProxyMaterial = MaterialInterface->GetMaterial();
-			EBlendMode BlendMode = MaterialInterface->GetBlendMode();
 			bool bIsMaterialUnlit = MaterialInterface->GetShadingModels().IsUnlit();
 			const bool bIsOpaque = IsOpaqueBlendMode(*MaterialInterface);
 			const bool bIsMasked = IsMaskedBlendMode(*MaterialInterface);
diff --git a/Engine/Source/Runtime/Engine/Private/DistanceFieldAtlas.cpp b/Engine/Source/Runtime/Engine/Private/DistanceFieldAtlas.cpp
--- a/Engine/Source/Runtime/Engine/Private/DistanceFieldAtlas.cpp
+++ b/Engine/Source/Runtime/Engine/Private/DistanceFieldAtlas.cpp
@@ -183,6 +183,7 @@ void FDistanceFieldVolumeData::CacheDerivedData(const FString& InStaticMeshDeriv
 		if (MaterialInterface)
 		{
 			MaterialData.BlendMode = MaterialInterface->GetBlendMode();
+			MaterialData.StrataBlendMode = MaterialInterface->GetStrataBlendMode();
 			MaterialData.bTwoSided = MaterialInterface->IsTwoSided();
 		}
 	}
@@ -191,8 +192,9 @@ void FDistanceFieldVolumeData::CacheDerivedData(const FString& InStaticMeshDeriv
 
 	for (int32 MaterialIndex = 0; MaterialIndex < Mesh->GetStaticMaterials().Num(); MaterialIndex++)
 	{
-		DistanceFieldKey += FString::Printf(TEXT("_M%u_%u_%u"), 
+		DistanceFieldKey += FString::Printf(TEXT("_M%u_%u_%u_%u"), 
 			(uint32)BuildMaterialData[MaterialIndex].BlendMode,
+			(uint32)BuildMaterialData[MaterialIndex].StrataBlendMode,
 			BuildMaterialData[MaterialIndex].bTwoSided ? 1 : 0,
 			BuildMaterialData[MaterialIndex].bAffectDistanceFieldLighting ? 1 : 0);
 	}
diff --git a/Engine/Source/Runtime/Engine/Private/Materials/MaterialShared.cpp b/Engine/Source/Runtime/Engine/Private/Materials/MaterialShared.cpp
--- a/Engine/Source/Runtime/Engine/Private/Materials/MaterialShared.cpp
+++ b/Engine/Source/Runtime/Engine/Private/Materials/MaterialShared.cpp
@@ -1687,7 +1687,7 @@ EMaterialFloatPrecisionMode FMaterialResource::GetMaterialFloatPrecisionMode() c
 
 bool FMaterialResource::IsUsingAlphaToCoverage() const
 {
-	return Material->bUseAlphaToCoverage && Material->MaterialDomain == EMaterialDomain::MD_Surface && Material->GetBlendMode() == EBlendMode::BLEND_Masked && !WritesEveryPixel();
+	return Material->bUseAlphaToCoverage && Material->MaterialDomain == EMaterialDomain::MD_Surface && IsMaskedBlendMode(*Material) && !WritesEveryPixel();
 }
 
 bool FMaterialResource::IsUsingPreintegratedGFForSimpleIBL() const
diff --git a/Engine/Source/Runtime/Engine/Private/MeshCardRepresentation.cpp b/Engine/Source/Runtime/Engine/Private/MeshCardRepresentation.cpp
--- a/Engine/Source/Runtime/Engine/Private/MeshCardRepresentation.cpp
+++ b/Engine/Source/Runtime/Engine/Private/MeshCardRepresentation.cpp
@@ -202,6 +202,7 @@ void FCardRepresentationData::CacheDerivedData(const FString& InDDCKey, const IT
 			if (MaterialInterface)
 			{
 				MaterialData.BlendMode = MaterialInterface->GetBlendMode();
+				MaterialData.StrataBlendMode = MaterialInterface->GetStrataBlendMode();
 				MaterialData.bTwoSided = MaterialInterface->IsTwoSided();
 			}
 		}
diff --git a/Engine/Source/Runtime/Renderer/Private/BasePassRendering.cpp b/Engine/Source/Runtime/Renderer/Private/BasePassRendering.cpp
--- a/Engine/Source/Runtime/Renderer/Private/BasePassRendering.cpp
+++ b/Engine/Source/Runtime/Renderer/Private/BasePassRendering.cpp
@@ -1874,8 +1874,7 @@ bool FBasePassMeshProcessor::ShouldDraw(const FMaterial& Material)
 
 				// only dual blended or modulate surfaces need background modulation
 			case ETranslucencyPass::TPT_TranslucencyAfterDOFModulate:
-				bShouldDraw = Material.IsTranslucencyAfterDOFEnabled() && (Material.IsDualBlendingEnabled(GetFeatureLevelShaderPlatform(FeatureLevel))  || 
-					BlendMode == BLEND_Modulate || StrataBlendMode == SBM_ColoredTransmittanceOnly);
+				bShouldDraw = Material.IsTranslucencyAfterDOFEnabled() && (Material.IsDualBlendingEnabled(GetFeatureLevelShaderPlatform(FeatureLevel)) || IsModulateBlendMode(BlendMode, StrataBlendMode));
 				break;
 
 			case ETranslucencyPass::TPT_TranslucencyAfterMotionBlur:
diff --git a/Engine/Source/Runtime/Renderer/Private/ShadowDepthRendering.cpp b/Engine/Source/Runtime/Renderer/Private/ShadowDepthRendering.cpp
--- a/Engine/Source/Runtime/Renderer/Private/ShadowDepthRendering.cpp
+++ b/Engine/Source/Runtime/Renderer/Private/ShadowDepthRendering.cpp
@@ -2033,7 +2033,6 @@ bool FShadowDepthPassMeshProcessor::TryAddMeshBatch(const FMeshBatch& RESTRICT M
 	const FMaterialRenderProxy& MaterialRenderProxy,
 	const FMaterial& Material)
 {
-	const EBlendMode BlendMode = Material.GetBlendMode();
 	const bool bShouldCastShadow = Material.ShouldCastDynamicShadows();
 
 	const FMeshDrawingPolicyOverrideSettings OverrideSettings = ComputeMeshOverrideSettings(MeshBatch);
diff --git a/Engine/Source/Runtime/Renderer/Private/SingleLayerWaterRendering.cpp b/Engine/Source/Runtime/Renderer/Private/SingleLayerWaterRendering.cpp
--- a/Engine/Source/Runtime/Renderer/Private/SingleLayerWaterRendering.cpp
+++ b/Engine/Source/Runtime/Renderer/Private/SingleLayerWaterRendering.cpp
@@ -1817,11 +1817,9 @@ void FSingleLayerWaterDepthPrepassMeshProcessor::CollectPSOInitializers(const FS
 		const FMeshDrawingPolicyOverrideSettings OverrideSettings = ComputeMeshOverrideSettings(PreCacheParams);
 		const ERasterizerFillMode MeshFillMode = ComputeMeshFillMode(Material, OverrideSettings);
 		const ERasterizerCullMode MeshCullMode = ComputeMeshCullMode(Material, OverrideSettings);
-		const EBlendMode BlendMode = Material.GetBlendMode();
-		const EStrataBlendMode StrataBlendMode = Material.GetStrataBlendMode();
 		const bool bSupportPositionOnlyStream = VertexFactoryType->SupportsPositionOnly();
 
-		if (IsOpaqueBlendMode(BlendMode, StrataBlendMode)
+		if (IsOpaqueBlendMode(Material)
 			&& bSupportPositionOnlyStream
 			&& !Material.MaterialModifiesMeshPosition_GameThread()
 			&& Material.WritesEveryPixel())
