World Partition - HLOD: Fix landscape HLODs being all black when using Nanite Landscapes
* Include Nanite landscape scene proxy in the scene capture as these are the ones that are actually visible when rendering LS with Nanite support

diff --git a/Engine/Source/Developer/MaterialUtilities/Private/MaterialUtilities.cpp b/Engine/Source/Developer/MaterialUtilities/Private/MaterialUtilities.cpp
--- a/Engine/Source/Developer/MaterialUtilities/Private/MaterialUtilities.cpp
+++ b/Engine/Source/Developer/MaterialUtilities/Private/MaterialUtilities.cpp
@@ -1011,17 +1011,36 @@ static bool ExportLandscapeMaterial(const ALandscapeProxy* InLandscape, const TS
 
 bool FMaterialUtilities::ExportLandscapeMaterial(const ALandscapeProxy* InLandscape, FFlattenMaterial& OutFlattenMaterial)
 {
-	TSet<FPrimitiveComponentId> ShowOnlyPrimitives;
+	bool bExportSuccess = false;
 
-	for (ULandscapeComponent* LandscapeComponent : InLandscape->LandscapeComponents)
+	if (InLandscape)
 	{
-		if (ensure(LandscapeComponent->SceneProxy))
+		TSet<FPrimitiveComponentId> ShowOnlyPrimitives;
+
+		// Include all landscape components scene proxies
+		for (ULandscapeComponent* LandscapeComponent : InLandscape->LandscapeComponents)
+		{
+			if (LandscapeComponent && LandscapeComponent->SceneProxy)
+			{
+				ShowOnlyPrimitives.Add(LandscapeComponent->SceneProxy->GetPrimitiveComponentId());
+			}
+		}
+
+		// Include Nanite landscape scene proxy - these are the ones that are actually visible when rendering LS with Nanite support
+		if (InLandscape->NaniteComponent && InLandscape->NaniteComponent->SceneProxy)
 		{
-			ShowOnlyPrimitives.Add(LandscapeComponent->SceneProxy->GetPrimitiveComponentId());
+			ShowOnlyPrimitives.Add(InLandscape->NaniteComponent->SceneProxy->GetPrimitiveComponentId());
 		}
+
+		bExportSuccess = ::ExportLandscapeMaterial(InLandscape, ShowOnlyPrimitives, {}, OutFlattenMaterial);
+	}
+	
+	if (!bExportSuccess)
+	{
+		UE_LOG(LogMaterialUtilities, Warning, TEXT("ExportLandscapeMaterial: Failed to export material for the provided ALandcapeProxy (%s)"), InLandscape ? *InLandscape->GetName() : TEXT("<null>"));
 	}
 
-	return ::ExportLandscapeMaterial(InLandscape, ShowOnlyPrimitives, {}, OutFlattenMaterial);
+	return bExportSuccess;
 }
 
 bool FMaterialUtilities::ExportLandscapeMaterial(const ALandscapeProxy* InLandscape, const TSet<FPrimitiveComponentId>& HiddenPrimitives, FFlattenMaterial& OutFlattenMaterial)
