Material Baking - When packing metallic/specular/roughness textures, also add "use" switches for these
* Fixes baking with some flattening materials that placed the UseXXX switch after testing for packed properties support

diff --git a/Engine/Source/Developer/MaterialUtilities/Private/ProxyMaterialUtilities.cpp b/Engine/Source/Developer/MaterialUtilities/Private/ProxyMaterialUtilities.cpp
--- a/Engine/Source/Developer/MaterialUtilities/Private/ProxyMaterialUtilities.cpp
+++ b/Engine/Source/Developer/MaterialUtilities/Private/ProxyMaterialUtilities.cpp
@@ -473,22 +473,21 @@ UMaterialInstanceConstant* FMaterialUtilities::CreateFlattenMaterialInstance(UPa
 		checkf(PackedTexture, TEXT("Failed to create texture"));
 		OutAssetsToSync.Add(PackedTexture);
 
+		auto PackProperty = [&NewStaticParameterSet](const FString& PropertyName, bool bEnable)
+		{
+			FStaticSwitchParameter SwitchParameter;
+			SwitchParameter.ParameterInfo.Name = FName(FString(TEXT("Pack")) + PropertyName);
+			SwitchParameter.Value = bEnable;
+			SwitchParameter.bOverride = true;
+			NewStaticParameterSet.StaticSwitchParameters.Add(SwitchParameter);
+			SwitchParameter.ParameterInfo.Name = FName(FString(TEXT("Use")) + PropertyName);
+			NewStaticParameterSet.StaticSwitchParameters.Add(SwitchParameter);
+		};
+		
 		// Setup switches for whether or not properties will be packed into one texture
-		FStaticSwitchParameter SwitchParameter;
-		SwitchParameter.ParameterInfo.Name = TEXT("PackMetallic");
-		SwitchParameter.Value = bPackMetallic;
-		SwitchParameter.bOverride = true;
-		NewStaticParameterSet.StaticSwitchParameters.Add(SwitchParameter);
-
-		SwitchParameter.ParameterInfo.Name = TEXT("PackSpecular");
-		SwitchParameter.Value = bPackSpecular;
-		SwitchParameter.bOverride = true;
-		NewStaticParameterSet.StaticSwitchParameters.Add(SwitchParameter);
-
-		SwitchParameter.ParameterInfo.Name = TEXT("PackRoughness");
-		SwitchParameter.Value = bPackRoughness;
-		SwitchParameter.bOverride = true;
-		NewStaticParameterSet.StaticSwitchParameters.Add(SwitchParameter);
+		PackProperty(TEXT("Metallic"), bPackMetallic);
+		PackProperty(TEXT("Specular"), bPackSpecular);
+		PackProperty(TEXT("Roughness"), bPackRoughness);
 
 		// Set up switch and texture values
 		OutMaterial->SetTextureParameterValueEditorOnly(PackedTextureName, PackedTexture);
