Cleaning up deprecated methods

diff --git a/Engine/Source/Developer/HierarchicalLODUtilities/Private/HierarchicalLODUtilities.cpp b/Engine/Source/Developer/HierarchicalLODUtilities/Private/HierarchicalLODUtilities.cpp
--- a/Engine/Source/Developer/HierarchicalLODUtilities/Private/HierarchicalLODUtilities.cpp
+++ b/Engine/Source/Developer/HierarchicalLODUtilities/Private/HierarchicalLODUtilities.cpp
@@ -58,11 +58,6 @@ DEFINE_LOG_CATEGORY_STATIC(LogHierarchicalLODUtilities, Verbose, All);
 
 #define LOCTEXT_NAMESPACE "HierarchicalLODUtils"
 
-void FHierarchicalLODUtilities::ExtractStaticMeshComponentsFromLODActor(AActor* Actor, TArray<UStaticMeshComponent*>& InOutComponents)
-{
-	UHLODProxy::ExtractStaticMeshComponentsFromLODActor(Cast<ALODActor>(Actor), InOutComponents);
-}
-
 void FHierarchicalLODUtilities::ExtractSubActorsFromLODActor(AActor* Actor, TArray<AActor*>& InOutActors)
 {
 	ALODActor* LODActor = CastChecked<ALODActor>(Actor);
@@ -224,28 +219,6 @@ UPackage* FHierarchicalLODUtilities::RetrieveLevelHLODPackage(const ULevel* InLe
 	return nullptr;
 }
 
-UPackage* FHierarchicalLODUtilities::CreateOrRetrieveLevelHLODPackage(const ULevel* InLevel)
-{
-	checkf(InLevel != nullptr, TEXT("Invalid Level supplied"));
-
-	UPackage* HLODPackage = nullptr;
-	UPackage* LevelOuterMost = InLevel->GetOutermost();
-
-	const FString PathName = FPackageName::GetLongPackagePath(LevelOuterMost->GetPathName());
-	const FString BaseName = FPackageName::GetShortName(LevelOuterMost->GetPathName());
-	const FString HLODLevelPackageName = FString::Printf(TEXT("%s/HLOD/%s_HLOD"), *PathName, *BaseName);
-
-	HLODPackage = CreatePackage( *HLODLevelPackageName);
-	HLODPackage->FullyLoad();
-	HLODPackage->Modify();
-	HLODPackage->SetPackageFlags(PKG_ContainsMapData);		// PKG_ContainsMapData required so FEditorFileUtils::GetDirtyContentPackages can treat this as a map package
-
-	// Target PackagePath; this is a hack to avoid save file dialog when we will be saving HLOD map package
-	HLODPackage->SetLoadedPath(FPackagePath::FromPackageNameChecked(HLODLevelPackageName));
-
-	return HLODPackage;
-}
-
 FString FHierarchicalLODUtilities::GetLevelHLODProxyName(const FString& InLevelPackageName, const uint32 InHLODLevelIndex)
 {
 	FString HLODProxyName;
@@ -253,12 +226,6 @@ FString FHierarchicalLODUtilities::GetLevelHLODProxyName(const FString& InLevelP
 	return HLODPackageName + TEXT(".") + HLODProxyName;
 }
 
-bool FHierarchicalLODUtilities::BuildStaticMeshForLODActor(ALODActor* LODActor, UPackage* AssetsOuter, const FHierarchicalSimplification& LODSetup, UMaterialInterface* InBaseMaterial)
-{
-	UHLODProxy* Proxy = FindObject<UHLODProxy>(AssetsOuter, *GetHLODProxyName(CastChecked<ULevel>(LODActor->GetOuter()), LODActor->LODLevel - 1));
-	return BuildStaticMeshForLODActor(LODActor, Proxy, LODSetup, InBaseMaterial);
-}
-
 static FString GetImposterMeshName(const UStaticMesh* InImposterMesh)
 {
 	UPackage* MeshOuterMost = InImposterMesh->GetOutermost();
@@ -881,13 +848,6 @@ bool FHierarchicalLODUtilities::BuildStaticMeshForLODActor(ALODActor* LODActor,
 	return true;
 }
 
-bool FHierarchicalLODUtilities::BuildStaticMeshForLODActor(ALODActor* LODActor, UPackage* AssetsOuter, const FHierarchicalSimplification& LODSetup)
-{
-PRAGMA_DISABLE_DEPRECATION_WARNINGS
-	return BuildStaticMeshForLODActor(LODActor, AssetsOuter, LODSetup, GEngine->DefaultHLODFlattenMaterial);
-PRAGMA_ENABLE_DEPRECATION_WARNINGS
-}
-
 EClusterGenerationError FHierarchicalLODUtilities::ShouldGenerateCluster(AActor* Actor, const int32 HLODLevelIndex)
 {
 	if (!Actor)
@@ -1016,11 +976,6 @@ void FHierarchicalLODUtilities::DestroyCluster(ALODActor* InActor)
 	}
 }
 
-void FHierarchicalLODUtilities::DestroyClusterData(ALODActor* InActor)
-{
-
-}
-
 ALODActor* FHierarchicalLODUtilities::CreateNewClusterActor(UWorld* InWorld, const int32 InLODLevel, AWorldSettings* WorldSettings)
 {
 	// Check incoming data
diff --git a/Engine/Source/Developer/HierarchicalLODUtilities/Public/HierarchicalLODUtilities.h b/Engine/Source/Developer/HierarchicalLODUtilities/Public/HierarchicalLODUtilities.h
--- a/Engine/Source/Developer/HierarchicalLODUtilities/Public/HierarchicalLODUtilities.h
+++ b/Engine/Source/Developer/HierarchicalLODUtilities/Public/HierarchicalLODUtilities.h
@@ -22,29 +22,20 @@ class FHierarchicalLODUtilities : public IHierarchicalLODUtilities
 {
 public:
 	~FHierarchicalLODUtilities() {}
-	UE_DEPRECATED(4.20, "Please use UHLODProxy::ExtractStaticMeshComponentsFromLODActor")
-	virtual void ExtractStaticMeshComponentsFromLODActor(AActor* Actor, TArray<UStaticMeshComponent*>& InOutComponents) override;
+	
 	virtual void ExtractSubActorsFromLODActor(AActor* Actor, TArray<AActor*>& InOutActors) override;
 	virtual float CalculateScreenSizeFromDrawDistance(const float SphereRadius, const FMatrix& ProjectionMatrix, const float Distance) override;
 	virtual float CalculateDrawDistanceFromScreenSize(const float SphereRadius, const float ScreenSize, const FMatrix& ProjectionMatrix) override;
-	UE_DEPRECATED(4.20, "This function signature is deprecated")
-	virtual UPackage* CreateOrRetrieveLevelHLODPackage(const ULevel* InLevel) override;
 	virtual UPackage* CreateOrRetrieveLevelHLODPackage(const ULevel* InLevel, const uint32 HLODLevelIndex) override;
 	virtual UPackage* RetrieveLevelHLODPackage(const ULevel* InLevel, const uint32 HLODLevelIndex) override;
 	virtual UHLODProxy* CreateOrRetrieveLevelHLODProxy(const ULevel* InLevel, const uint32 HLODLevelIndex) override;
 	virtual UHLODProxy* RetrieveLevelHLODProxy(const ULevel* InLevel, const uint32 HLODLevelIndex) override;
 	virtual FString GetWildcardOfHLODPackagesForPackage(const FString& PackageName) override;
 	virtual FString GetLevelHLODProxyName(const FString& InLevelPackageName, const uint32 HLODLevelIndex) override;
-	UE_DEPRECATED(4.20, "Use different signature containing BaseMaterial parameter")
-	virtual bool BuildStaticMeshForLODActor(ALODActor* LODActor, UPackage* AssetsOuter, const FHierarchicalSimplification& LODSetup) override;
-	UE_DEPRECATED(4.20, "Use different signature using UHLODProxy")
-	virtual bool BuildStaticMeshForLODActor(ALODActor* LODActor, UPackage* AssetsOuter, const FHierarchicalSimplification& LODSetup, UMaterialInterface* InBaseMaterial) override;
 	virtual bool BuildStaticMeshForLODActor(ALODActor* LODActor, UHLODProxy* Proxy, const FHierarchicalSimplification& LODSetup, UMaterialInterface* InBaseMaterial) override;
 	virtual EClusterGenerationError ShouldGenerateCluster(AActor* Actor, const int32 HLODLevelIndex) override;
 	virtual ALODActor* GetParentLODActor(const AActor* InActor) override;
 	virtual void DestroyCluster(ALODActor* InActor) override;
-	UE_DEPRECATED(4.20, "This function is no longer used")
-	virtual void DestroyClusterData(ALODActor* InActor) override;
 	virtual ALODActor* CreateNewClusterActor(UWorld* InWorld, const int32 InLODLevel, AWorldSettings* WorldSettings) override;
 	virtual ALODActor* CreateNewClusterFromActors(UWorld* InWorld, AWorldSettings* WorldSettings, const TArray<AActor*>& InActors, const int32 InLODLevel = 0) override;
 	virtual const bool RemoveActorFromCluster(AActor* InActor) override;
diff --git a/Engine/Source/Developer/HierarchicalLODUtilities/Public/IHierarchicalLODUtilities.h b/Engine/Source/Developer/HierarchicalLODUtilities/Public/IHierarchicalLODUtilities.h
--- a/Engine/Source/Developer/HierarchicalLODUtilities/Public/IHierarchicalLODUtilities.h
+++ b/Engine/Source/Developer/HierarchicalLODUtilities/Public/IHierarchicalLODUtilities.h
@@ -40,14 +40,6 @@ class HIERARCHICALLODUTILITIES_API IHierarchicalLODUtilities
 public:	
 	virtual ~IHierarchicalLODUtilities() {}
 
-	/**
-	* Recursively retrieves StaticMeshComponents from a LODActor and its child LODActors
-	*
-	* @param Actor - LODActor instance
-	* @param InOutComponents - Will hold the StaticMeshComponents
-	*/
-	virtual void ExtractStaticMeshComponentsFromLODActor(AActor* Actor, TArray<UStaticMeshComponent*>& InOutComponents) = 0;
-
 	/**
 	* Recursively retrieves Actors from a LODActor and its child LODActors
 	*
@@ -61,8 +53,6 @@ public:
 
 	virtual float CalculateDrawDistanceFromScreenSize(const float SphereRadius, const float ScreenSize, const FMatrix& ProjectionMatrix) = 0;
 
-	virtual UPackage* CreateOrRetrieveLevelHLODPackage(const ULevel* InLevel) = 0;
-
 	/** Creates or retrieves the HLOD package that is created for the given level */
 	virtual UPackage* CreateOrRetrieveLevelHLODPackage(const ULevel* InLevel, const uint32 HLODLevelIndex) = 0;
 
@@ -88,8 +78,6 @@ public:
 	* @param Outer - Outer object to store the mesh in
 	* @return UStaticMesh*
 	*/
-	virtual bool BuildStaticMeshForLODActor(ALODActor* LODActor, UPackage* AssetsOuter, const FHierarchicalSimplification& LODSetup) = 0;
-	virtual bool BuildStaticMeshForLODActor(ALODActor* LODActor, UPackage* AssetsOuter, const FHierarchicalSimplification& LODSetup, UMaterialInterface* InBaseMaterial) = 0;
 	virtual bool BuildStaticMeshForLODActor(ALODActor* LODActor, UHLODProxy* Proxy, const FHierarchicalSimplification& LODSetup, UMaterialInterface* InBaseMaterial) = 0;
 	
 	/**
@@ -106,9 +94,6 @@ public:
 	/** Deletes the given cluster's data and instance in the world */
 	virtual void DestroyCluster(ALODActor* InActor) = 0;
 
-	/** Deletes the given cluster's assets */
-	virtual void DestroyClusterData(ALODActor* InActor) = 0;
-
 	/** Creates a new cluster actor in the given InWorld with InLODLevel as HLODLevel */
 	virtual ALODActor* CreateNewClusterActor(UWorld* InWorld, const int32 InLODLevel, AWorldSettings* WorldSettings) = 0;
 
diff --git a/Engine/Source/Developer/MaterialUtilities/Private/MaterialUtilities.cpp b/Engine/Source/Developer/MaterialUtilities/Private/MaterialUtilities.cpp
--- a/Engine/Source/Developer/MaterialUtilities/Private/MaterialUtilities.cpp
+++ b/Engine/Source/Developer/MaterialUtilities/Private/MaterialUtilities.cpp
@@ -1497,7 +1497,7 @@ UMaterialInstanceConstant* FMaterialUtilities::CreateInstancedMaterial(UMaterial
 	const FString AssetBasePath = InOuter ? TEXT("") : FPackageName::GetLongPackagePath(BaseName);
 
 	// Create material
-	const FString MaterialAssetName = TEXT("M_") + AssetBaseName;
+	const FString MaterialAssetName = TEXT("MI_") + AssetBaseName;
 	UPackage* MaterialOuter = InOuter;
 	if (MaterialOuter == NULL)
 	{		
@@ -2351,273 +2351,6 @@ bool FMaterialUtilities::ExportMaterialUVDensities(UMaterialInterface* InMateria
 	return true;
 }
 
-bool FMaterialUtilities::ExportMaterial(struct FMaterialMergeData& InMaterialData, FFlattenMaterial& OutFlattenMaterial, struct FExportMaterialProxyCache* ProxyCache)
-{
-	UMaterialInterface* Material = InMaterialData.Material;
-	UE_LOG(LogMaterialUtilities, Log, TEXT("Flattening material: %s"), *Material->GetName());
-
-	if (ProxyCache)
-	{
-		// ExportMaterial was called with non-null CompiledMaterial. This means compiled shaders
-		// should be stored outside, and could be re-used in next call to ExportMaterial.
-		// FMaterialData already has "proxy cache" fiels, should swap it with CompiledMaterial,
-		// and swap back before returning from this function.
-		// Purpose of the following line: use compiled material cached from previous call.
-		Exchange(ProxyCache, InMaterialData.ProxyCache);
-	}
-	
-	// Precache all used textures, otherwise could get everything rendered with low-res textures.
-	TArray<UTexture*> MaterialTextures;
-	Material->GetUsedTextures(MaterialTextures, EMaterialQualityLevel::Num, true, GMaxRHIFeatureLevel, true);
-
-	FTextureCompilingManager::Get().FinishCompilation(MaterialTextures);
-
-	for (UTexture* Texture : MaterialTextures)
-	{
-		if (Texture != NULL)
-		{
-			UTexture2D* Texture2D = Cast<UTexture2D>(Texture);
-			if (Texture2D)
-			{
-				Texture2D->SetForceMipLevelsToBeResident(30.0f, true);
-				Texture2D->WaitForStreaming();
-			}
-		}		
-	}
-
-	// Determine whether or not certain properties can be rendered
-	const bool bRenderNormal = (Material->GetMaterial()->HasNormalConnected() || Material->GetMaterial()->bUseMaterialAttributes) && OutFlattenMaterial.ShouldGenerateDataForProperty(EFlattenMaterialProperties::Normal);
-	const bool bRenderTangent = (Material->GetMaterial()->IsPropertyConnected(MP_Tangent) || Material->GetMaterial()->bUseMaterialAttributes) && OutFlattenMaterial.ShouldGenerateDataForProperty(EFlattenMaterialProperties::Tangent);
-	const bool bRenderEmissive = (Material->GetMaterial()->IsPropertyConnected(MP_EmissiveColor) || Material->GetMaterial()->bUseMaterialAttributes) && OutFlattenMaterial.ShouldGenerateDataForProperty(EFlattenMaterialProperties::Emissive);
-	const bool bRenderOpacityMask = Material->IsPropertyActive(MP_OpacityMask) && IsMaskedBlendMode(*Material) && OutFlattenMaterial.ShouldGenerateDataForProperty(EFlattenMaterialProperties::OpacityMask);
-	const bool bRenderOpacity = Material->IsPropertyActive(MP_Opacity) && IsTranslucentBlendMode(*Material) && OutFlattenMaterial.ShouldGenerateDataForProperty(EFlattenMaterialProperties::Opacity);
-	const bool bRenderSubSurface = Material->IsPropertyActive(MP_SubsurfaceColor) && OutFlattenMaterial.ShouldGenerateDataForProperty(EFlattenMaterialProperties::SubSurface);
-	const bool bRenderMetallic = Material->IsPropertyActive(MP_Metallic) && OutFlattenMaterial.ShouldGenerateDataForProperty(EFlattenMaterialProperties::Metallic);
-	const bool bRenderSpecular = Material->IsPropertyActive(MP_Specular) && OutFlattenMaterial.ShouldGenerateDataForProperty(EFlattenMaterialProperties::Specular);
-	const bool bRenderRoughness = Material->IsPropertyActive(MP_Roughness) && OutFlattenMaterial.ShouldGenerateDataForProperty(EFlattenMaterialProperties::Roughness);
-	const bool bRenderAnisotropy = Material->IsPropertyActive(MP_Anisotropy) && OutFlattenMaterial.ShouldGenerateDataForProperty(EFlattenMaterialProperties::Anisotropy);
-
-	check(!bRenderOpacity || !bRenderOpacityMask);
-
-	FIntPoint Size;
-	// Compile shaders and render flatten material.
-	Size = OutFlattenMaterial.GetPropertySize(EFlattenMaterialProperties::Diffuse);
-	PRAGMA_DISABLE_DEPRECATION_WARNINGS
-	RenderMaterialPropertyToTexture(InMaterialData, MP_BaseColor, false, PF_B8G8R8A8, OutFlattenMaterial.RenderSize, Size, OutFlattenMaterial.GetPropertySamples(EFlattenMaterialProperties::Diffuse));
-	PRAGMA_ENABLE_DEPRECATION_WARNINGS
-	OutFlattenMaterial.SetPropertySize(EFlattenMaterialProperties::Diffuse, Size);
-
-	if (bRenderMetallic)
-	{
-		Size = OutFlattenMaterial.GetPropertySize(EFlattenMaterialProperties::Metallic);
-		PRAGMA_DISABLE_DEPRECATION_WARNINGS
-		RenderMaterialPropertyToTexture(InMaterialData, MP_Metallic, false, PF_B8G8R8A8, OutFlattenMaterial.RenderSize, Size, 
-			OutFlattenMaterial.GetPropertySamples(EFlattenMaterialProperties::Metallic));
-		PRAGMA_ENABLE_DEPRECATION_WARNINGS
-		OutFlattenMaterial.SetPropertySize(EFlattenMaterialProperties::Metallic, Size);
-	}
-
-	if (bRenderSpecular)
-	{
-		Size = OutFlattenMaterial.GetPropertySize(EFlattenMaterialProperties::Specular);
-		PRAGMA_DISABLE_DEPRECATION_WARNINGS
-		RenderMaterialPropertyToTexture(InMaterialData, MP_Specular, false, PF_B8G8R8A8, OutFlattenMaterial.RenderSize, Size, OutFlattenMaterial.GetPropertySamples(EFlattenMaterialProperties::Specular));
-		PRAGMA_ENABLE_DEPRECATION_WARNINGS
-		OutFlattenMaterial.SetPropertySize(EFlattenMaterialProperties::Specular, Size);
-	}
-
-	if (bRenderRoughness)
-	{
-		Size = OutFlattenMaterial.GetPropertySize(EFlattenMaterialProperties::Roughness);
-		PRAGMA_DISABLE_DEPRECATION_WARNINGS
-		RenderMaterialPropertyToTexture(InMaterialData, MP_Roughness, false, PF_B8G8R8A8, OutFlattenMaterial.RenderSize, Size, OutFlattenMaterial.GetPropertySamples(EFlattenMaterialProperties::Roughness));
-		PRAGMA_ENABLE_DEPRECATION_WARNINGS
-		OutFlattenMaterial.SetPropertySize(EFlattenMaterialProperties::Roughness, Size);
-	}
-
-	if (bRenderAnisotropy)
-	{
-		Size = OutFlattenMaterial.GetPropertySize(EFlattenMaterialProperties::Anisotropy);
-		PRAGMA_DISABLE_DEPRECATION_WARNINGS
-		RenderMaterialPropertyToTexture(InMaterialData, MP_Anisotropy, false, PF_B8G8R8A8, OutFlattenMaterial.RenderSize, Size, OutFlattenMaterial.GetPropertySamples(EFlattenMaterialProperties::Anisotropy));
-		PRAGMA_ENABLE_DEPRECATION_WARNINGS
-		OutFlattenMaterial.SetPropertySize(EFlattenMaterialProperties::Anisotropy, Size);
-	}
-
-	if (bRenderNormal)
-	{
-		Size = OutFlattenMaterial.GetPropertySize(EFlattenMaterialProperties::Normal);
-		PRAGMA_DISABLE_DEPRECATION_WARNINGS
-		RenderMaterialPropertyToTexture(InMaterialData, MP_Normal, true, PF_B8G8R8A8, OutFlattenMaterial.RenderSize, Size, OutFlattenMaterial.GetPropertySamples(EFlattenMaterialProperties::Normal));
-		PRAGMA_ENABLE_DEPRECATION_WARNINGS
-		OutFlattenMaterial.SetPropertySize(EFlattenMaterialProperties::Normal, Size);
-	}
-	else
-	{
-        // Make sure we output a default normal value in case the material does not generate one (to prevent issues with combining meshes with and without normal maps being atlassed together)
-		TArray<FColor>& Samples = OutFlattenMaterial.GetPropertySamples(EFlattenMaterialProperties::Normal);
-		Samples.Add(FColor(128, 128, 255));
-		OutFlattenMaterial.SetPropertySize(EFlattenMaterialProperties::Normal, FIntPoint(1,1));
-	}
-
-	if (bRenderTangent)
-	{
-		Size = OutFlattenMaterial.GetPropertySize(EFlattenMaterialProperties::Tangent);
-		PRAGMA_DISABLE_DEPRECATION_WARNINGS
-		RenderMaterialPropertyToTexture(InMaterialData, MP_Tangent, true, PF_B8G8R8A8, OutFlattenMaterial.RenderSize, Size, OutFlattenMaterial.GetPropertySamples(EFlattenMaterialProperties::Tangent));
-		PRAGMA_ENABLE_DEPRECATION_WARNINGS
-		OutFlattenMaterial.SetPropertySize(EFlattenMaterialProperties::Tangent, Size);
-	}
-	else
-	{
-		// Make sure we output a default tangent value in case the material does not generate one (to prevent issues with combining meshes with and without tangent maps being atlassed together)
-		TArray<FColor>& Samples = OutFlattenMaterial.GetPropertySamples(EFlattenMaterialProperties::Tangent);
-		Samples.Add(FColor(255, 128, 128));
-		OutFlattenMaterial.SetPropertySize(EFlattenMaterialProperties::Tangent, FIntPoint(1, 1));
-	}
-
-	if (bRenderOpacityMask)
-	{
-		Size = OutFlattenMaterial.GetPropertySize(EFlattenMaterialProperties::OpacityMask);
-		PRAGMA_DISABLE_DEPRECATION_WARNINGS
-		RenderMaterialPropertyToTexture(InMaterialData, MP_OpacityMask, true, PF_B8G8R8A8, OutFlattenMaterial.RenderSize, Size, OutFlattenMaterial.GetPropertySamples(EFlattenMaterialProperties::OpacityMask));
-		PRAGMA_ENABLE_DEPRECATION_WARNINGS
-		OutFlattenMaterial.SetPropertySize(EFlattenMaterialProperties::OpacityMask, Size);
-	}
-
-	if (bRenderOpacity)
-	{
-		Size = OutFlattenMaterial.GetPropertySize(EFlattenMaterialProperties::Opacity);
-		// Number of blend modes, let's UMaterial decide whether it wants this property
-		PRAGMA_DISABLE_DEPRECATION_WARNINGS
-		RenderMaterialPropertyToTexture(InMaterialData, MP_Opacity, true, PF_B8G8R8A8, OutFlattenMaterial.RenderSize, Size, OutFlattenMaterial.GetPropertySamples(EFlattenMaterialProperties::Opacity));
-		PRAGMA_ENABLE_DEPRECATION_WARNINGS
-		OutFlattenMaterial.SetPropertySize(EFlattenMaterialProperties::Opacity, Size);
-	}
-
-	if (bRenderEmissive)
-	{
-		Size = OutFlattenMaterial.GetPropertySize(EFlattenMaterialProperties::Emissive);
-		// PF_FloatRGBA is here to be able to render and read HDR image using ReadFloat16Pixels()
-		PRAGMA_DISABLE_DEPRECATION_WARNINGS
-		RenderMaterialPropertyToTexture(InMaterialData, MP_EmissiveColor, false, PF_FloatRGBA, OutFlattenMaterial.RenderSize, Size, OutFlattenMaterial.GetPropertySamples(EFlattenMaterialProperties::Emissive));
-		PRAGMA_ENABLE_DEPRECATION_WARNINGS
-		OutFlattenMaterial.EmissiveScale = InMaterialData.EmissiveScale;
-		OutFlattenMaterial.SetPropertySize(EFlattenMaterialProperties::Emissive, Size);
-	}	
-
-	if (bRenderSubSurface)
-	{
-		// TODO support rendering out sub surface color property
-		/*RenderMaterialPropertyToTexture(InMaterialData, MP_SubsurfaceColor, false, PF_B8G8R8A8, OutFlattenMaterial.RenderSize, OutFlattenMaterial.SubSurfaceSize, OutFlattenMaterial.SubSurfaceSamples);*/
-	}
-
-	OutFlattenMaterial.MaterialId = Material->GetLightingGuid();
-
-	// Swap back the proxy cache
-	if (ProxyCache)
-	{
-		// Store compiled material to external cache.
-		Exchange(ProxyCache, InMaterialData.ProxyCache);
-	}
-
-	UE_LOG(LogMaterialUtilities, Log, TEXT("Material flattening done. (%s)"), *Material->GetName());
-
-	return true;
-}
-
-bool FMaterialUtilities::RenderMaterialPropertyToTexture(struct FMaterialMergeData& InMaterialData, EMaterialProperty InMaterialProperty, bool bInForceLinearGamma, EPixelFormat InPixelFormat, const FIntPoint InTargetSize, FIntPoint& OutSampleSize, TArray<FColor>& OutSamples)
-{
-	if (InTargetSize.X == 0 || InTargetSize.Y == 0)
-	{
-		return false;
-	}
-
-	OutSampleSize = InTargetSize;
-
-	FMaterialRenderProxy* MaterialProxy = nullptr;
-
-	check(InMaterialProperty >= 0 && InMaterialProperty < UE_ARRAY_COUNT(InMaterialData.ProxyCache->Proxies));
-	if (InMaterialData.ProxyCache->Proxies[InMaterialProperty])
-	{
-		MaterialProxy = InMaterialData.ProxyCache->Proxies[InMaterialProperty];
-	}
-	else
-	{
-		MaterialProxy = InMaterialData.ProxyCache->Proxies[InMaterialProperty] = new FExportMaterialProxy(InMaterialData.Material, InMaterialProperty);
-	}
-	
-	if (MaterialProxy == nullptr)
-	{
-		return false;
-	}
-	
-	// Disallow garbage collection of RenderTarget.
-	check(CurrentlyRendering == false);
-	CurrentlyRendering = true;
-
-	const bool bNormalMap = (InMaterialProperty == MP_Normal);
-	UTextureRenderTarget2D* RenderTarget = CreateRenderTarget(bInForceLinearGamma, bNormalMap, InPixelFormat, OutSampleSize);
-	OutSamples.Empty(InTargetSize.X * InTargetSize.Y);
-	bool bResult = FMeshRenderer::RenderMaterial(
-		InMaterialData,
-		MaterialProxy,
-		InMaterialProperty,
-		RenderTarget,
-		OutSamples);
-
-	/** Disabled for now, see comment below */
-	// Check for uniform value, perhaps this can be determined before rendering the material, see WillGenerateUniformData (LightmassRender)
-	/*bool bIsUniform = true;
-	FColor MaxColor(0, 0, 0, 0);
-	if (bResult)
-	{
-		// Find maximal color value
-		int32 MaxColorValue = 0;
-		for (int32 Index = 0; Index < OutSamples.Num(); Index++)
-		{
-			FColor Color = OutSamples[Index];
-			int32 ColorValue = Color.R + Color.G + Color.B + Color.A;
-			if (ColorValue > MaxColorValue)
-			{
-				MaxColorValue = ColorValue;
-				MaxColor = Color;
-			}
-		}
-
-		// Fill background with maximal color value and render again		
-		RenderTarget->ClearColor = FLinearColor(MaxColor);
-		TArray<FColor> OutSamples2;
-		FMeshRenderer::RenderMaterial(
-			InMaterialData,
-			MaterialProxy,
-			InMaterialProperty,
-			RenderTarget,
-			OutSamples2);
-		for (int32 Index = 0; Index < OutSamples2.Num(); Index++)
-		{
-			FColor Color = OutSamples2[Index];
-			if (Color != MaxColor)
-			{
-				bIsUniform = false;
-				break;
-			}
-		}
-	}
-
-	// Uniform value
-	if (bIsUniform)
-	{
-		OutSampleSize = FIntPoint(1, 1);
-		OutSamples.Empty();
-		OutSamples.Add(MaxColor);
-	}*/
-
-	CurrentlyRendering = false;
-
-	return bResult;
-}
-
 UTextureRenderTarget2D* FMaterialUtilities::CreateRenderTarget(bool bInForceLinearGamma, bool bNormalMap, EPixelFormat InPixelFormat, FIntPoint& InTargetSize)
 {
 	const FLinearColor ClearColour = bNormalMap ? FLinearColor(0.0f, 0.0f, 0.0f, 0.0f) : FLinearColor(1.0f, 0.0f, 1.0f, 0.0f);
diff --git a/Engine/Source/Developer/MaterialUtilities/Public/MaterialUtilities.h b/Engine/Source/Developer/MaterialUtilities/Public/MaterialUtilities.h
--- a/Engine/Source/Developer/MaterialUtilities/Public/MaterialUtilities.h
+++ b/Engine/Source/Developer/MaterialUtilities/Public/MaterialUtilities.h
@@ -536,31 +536,6 @@ public:
 
 private:
 	
-	/**
-	* Private export material function to which unique signatures are mapped for unified code path
-	*
-	* @param InMaterialData			Target material data 
-	* @param OutFlattenMaterial		Output flattened material
-	* @return						Whether operation was successful
-	*/
-	UE_DEPRECATED(4.19, "Please use new functionality in MaterialBaking module")
-	static bool ExportMaterial(struct FMaterialMergeData& InMaterialData, FFlattenMaterial& OutFlattenMaterial, struct FExportMaterialProxyCache* ProxyCache = nullptr);
-
-	/**
-	* Renders out the specified material property with the given material data to a texture
-	*
-	* @param InMaterialData			Target material data
-	* @param InMaterialProperty		Target material property
-	* @param bInForceLinearGamma	Whether or not to force linear gamma (used for Normal property)
-	* @param InPixelFormat			Pixel format of the target texture
-	* @param InTargetSize			Dimensions of the target texture
-	* @param OutSampleSize			Dimensions of the rendered texture
-	* @param OutSamples				Array of FColor samples containing the rendered out texture pixel data
-	* @return						Whether operation was successful
-	*/
-	UE_DEPRECATED(4.19, "Please use new functionality in MaterialBaking module")
-	static bool RenderMaterialPropertyToTexture(struct FMaterialMergeData& InMaterialData, EMaterialProperty InMaterialProperty, bool bInForceLinearGamma, EPixelFormat InPixelFormat, const FIntPoint InTargetSize, FIntPoint& OutSampleSize, TArray<FColor>& OutSamples);
-
 	/**
 	* Creates and add or reuses a RenderTarget from the pool
 	*
diff --git a/Engine/Source/Developer/MeshMergeUtilities/Private/MeshMergeUtilities.cpp b/Engine/Source/Developer/MeshMergeUtilities/Private/MeshMergeUtilities.cpp
--- a/Engine/Source/Developer/MeshMergeUtilities/Private/MeshMergeUtilities.cpp
+++ b/Engine/Source/Developer/MeshMergeUtilities/Private/MeshMergeUtilities.cpp
@@ -2002,22 +2002,6 @@ void FMeshMergeUtilities::CreateProxyMesh(const TArray<UStaticMeshComponent*>& I
 	);
 }
 
-bool FMeshMergeUtilities::IsValidBaseMaterial(const UMaterialInterface* InBaseMaterial, bool bShowToaster) const
-{
-	bool bIsValid = FMaterialUtilities::IsValidFlattenMaterial(InBaseMaterial);
-	if (!bIsValid && bShowToaster)
-	{
-		FFormatNamedArguments Arguments;
-		Arguments.Add(TEXT("MaterialName"), FText::FromString(InBaseMaterial->GetName()));
-		FText ErrorMessage = FText::Format(LOCTEXT("UHierarchicalLODSettings_PostEditChangeProperty", "Material {MaterialName} is missing required Material Parameters (check log for details)"), Arguments);
-		FNotificationInfo Info(ErrorMessage);
-		Info.ExpireDuration = 5.0f;
-		FSlateNotificationManager::Get().AddNotification(Info);
-	}
-
-	return bIsValid;
-}
-
 void FMeshMergeUtilities::RetrieveMeshDescription(const UStaticMeshComponent* InStaticMeshComponent, int32 LODIndex, FMeshDescription& InOutMeshDescription, bool bPropagateMeshData) const
 {
 	FMeshMergeHelpers::RetrieveMesh(InStaticMeshComponent, LODIndex, InOutMeshDescription, bPropagateMeshData);
diff --git a/Engine/Source/Developer/MeshMergeUtilities/Private/MeshMergeUtilities.h b/Engine/Source/Developer/MeshMergeUtilities/Private/MeshMergeUtilities.h
--- a/Engine/Source/Developer/MeshMergeUtilities/Private/MeshMergeUtilities.h
+++ b/Engine/Source/Developer/MeshMergeUtilities/Private/MeshMergeUtilities.h
@@ -34,9 +34,6 @@ public:
 	virtual void CreateProxyMesh(const TArray<AActor*>& InActors, const struct FMeshProxySettings& InMeshProxySettings, UPackage* InOuter, const FString& InProxyBasePackageName, const FGuid InGuid, const FCreateProxyDelegate& InProxyCreatedDelegate, const bool bAllowAsync = false, const float ScreenSize = 1.0f) const override;
 	virtual void CreateProxyMesh(const TArray<UStaticMeshComponent*>& InStaticMeshComps, const struct FMeshProxySettings& InMeshProxySettings, UPackage* InOuter, const FString& InProxyBasePackageName, const FGuid InGuid, const FCreateProxyDelegate& InProxyCreatedDelegate, const bool bAllowAsync = false, const float ScreenSize = 1.0f) const override;
 
-	UE_DEPRECATED(5.0, "Use FMaterialUtilities::IsValidFlattenMaterial()")
-	virtual bool IsValidBaseMaterial(const UMaterialInterface* InBaseMaterial, bool bShowToaster) const override;
-
 	virtual void RetrieveMeshDescription(const UStaticMeshComponent* InStaticMeshComponent, int32 LODIndex, FMeshDescription& InOutMeshDescription, bool bPropagateVertexColours) const override;
 	virtual void RetrieveMeshDescription(const USkeletalMeshComponent* InSkeletalMeshComponent, int32 LODIndex, FMeshDescription& InOutMeshDescription, bool bPropagateVertexColours) const override;
 	virtual void RetrieveMeshDescription(const UStaticMesh* InStaticMesh, int32 LODIndex, FMeshDescription& InOutMeshDescription) const override;
diff --git a/Engine/Source/Developer/MeshMergeUtilities/Public/IMeshMergeUtilities.h b/Engine/Source/Developer/MeshMergeUtilities/Public/IMeshMergeUtilities.h
--- a/Engine/Source/Developer/MeshMergeUtilities/Public/IMeshMergeUtilities.h
+++ b/Engine/Source/Developer/MeshMergeUtilities/Public/IMeshMergeUtilities.h
@@ -125,9 +125,6 @@ public:
 	*/
 	virtual void CreateProxyMesh(const TArray<UStaticMeshComponent*>& InComponents, const struct FMeshProxySettings& InMeshProxySettings, UMaterialInterface* InBaseMaterial, UPackage* InOuter, const FString& InProxyBasePackageName, const FGuid InGuid, const FCreateProxyDelegate& InProxyCreatedDelegate, const bool bAllowAsync = false, const float ScreenSize = 1.0f) const = 0;
 
-	UE_DEPRECATED(5.0, "Use FMaterialUtilities::IsValidFlattenMaterial()")
-	virtual bool IsValidBaseMaterial(const UMaterialInterface* InBaseMaterial, bool bShowToaster) const = 0;
-
 	/** Extract mesh data in FMeshDescription form from the provided static mesh component */
 	virtual void RetrieveMeshDescription(const UStaticMeshComponent* InStaticMeshComponent, int32 LODIndex, FMeshDescription& InOutMeshDescription, bool bPropagateVertexColours) const = 0;
 	/** Extract mesh data in FMeshDescription form from the provided skeletal mesh component */
diff --git a/Engine/Source/Editor/UnrealEd/Private/HierarchicalLOD.cpp b/Engine/Source/Editor/UnrealEd/Private/HierarchicalLOD.cpp
--- a/Engine/Source/Editor/UnrealEd/Private/HierarchicalLOD.cpp
+++ b/Engine/Source/Editor/UnrealEd/Private/HierarchicalLOD.cpp
@@ -854,7 +854,7 @@ void FHierarchicalLODBuilder::BuildMeshesForLODActors(bool bForceAll)
 							{
 								SlowTask.EnterProgressFrame(100.0f / (float)NumLODActors, FText::Format(LOCTEXT("HierarchicalLOD_BuildLODActorMeshesProgress", "Building LODActor Mesh {0} of {1} (LOD Level {2})"), FText::AsNumber(LODActorIndex), FText::AsNumber(LODLevelActors[LODIndex].Num()), FText::AsNumber(LODIndex + 1)));
 
-								bool bBuildSuccessful = Utilities->BuildStaticMeshForLODActor(Actor, AssetsOuter, BuildLODLevelSettings[LODIndex], BaseMaterial);
+								bool bBuildSuccessful = Utilities->BuildStaticMeshForLODActor(Actor, Proxy, BuildLODLevelSettings[LODIndex], BaseMaterial);
 
 								// Report an error if the build failed
 								if (!bBuildSuccessful)
@@ -1152,8 +1152,7 @@ void FHierarchicalLODBuilder::BuildMeshForLODActor(ALODActor* LODActor, const ui
 	IHierarchicalLODUtilities* Utilities = Module.GetUtilities();
 
 	UHLODProxy* Proxy = Utilities->CreateOrRetrieveLevelHLODProxy(LODActor->GetLevel(), LODLevel);
-	UPackage* AssetsOuter = Proxy->GetOutermost();
-	const bool bResult = Utilities->BuildStaticMeshForLODActor(LODActor, AssetsOuter, BuildLODLevelSettings[LODLevel], BaseMaterial);
+	const bool bResult = Utilities->BuildStaticMeshForLODActor(LODActor, Proxy, BuildLODLevelSettings[LODLevel], BaseMaterial);
 
 	if (bResult == false)
 	{
@@ -1328,9 +1327,9 @@ ALODActor* FHierarchicalLODBuilder::CreateLODActor(const FLODCluster& InCluster,
 		{
 			TArray<UStaticMeshComponent*> Components;
 
-			if (Actor->IsA<ALODActor>())
+			if (ALODActor* LODActor = Cast<ALODActor>(Actor))
 			{
-				Utilities->ExtractStaticMeshComponentsFromLODActor(Actor, Components);
+				UHLODProxy::ExtractStaticMeshComponentsFromLODActor(LODActor, Components);
 			}
 			else
 			{
diff --git a/Engine/Source/Runtime/Engine/Classes/Engine/HLODProxy.h b/Engine/Source/Runtime/Engine/Classes/Engine/HLODProxy.h
--- a/Engine/Source/Runtime/Engine/Classes/Engine/HLODProxy.h
+++ b/Engine/Source/Runtime/Engine/Classes/Engine/HLODProxy.h
@@ -59,11 +59,7 @@ public:
 	static ENGINE_API uint32 GetCRC(UStaticMeshComponent* InComponent, uint32 InCRC = 0, const FTransform& TransformComponents = FTransform::Identity);
 
 	ENGINE_API virtual void PostLoad() override;
-	PRAGMA_DISABLE_DEPRECATION_WARNINGS // Suppress compiler warning on override of deprecated function
-	UE_DEPRECATED(5.0, "Use version that takes FObjectPreSaveContext instead.")
-	ENGINE_API virtual void PreSave(const class ITargetPlatform* TargetPlatform) override;
-	ENGINE_API PRAGMA_ENABLE_DEPRECATION_WARNINGS
-	virtual void PreSave(FObjectPreSaveContext ObjectSaveContext) override;
+	ENGINE_API virtual void PreSave(FObjectPreSaveContext ObjectSaveContext) override;
 
 	/** Returns true if proxy doesn't contain any mesh entry. */
 	ENGINE_API bool IsEmpty() const;
diff --git a/Engine/Source/Runtime/Engine/Classes/Engine/LODActor.h b/Engine/Source/Runtime/Engine/Classes/Engine/LODActor.h
--- a/Engine/Source/Runtime/Engine/Classes/Engine/LODActor.h
+++ b/Engine/Source/Runtime/Engine/Classes/Engine/LODActor.h
@@ -131,9 +131,6 @@ public:
 	ENGINE_API void AddInstances(const UStaticMesh* InStaticMesh, const UMaterialInterface* InMaterial, const TArray<FTransform>& InTransforms);
 	ENGINE_API void AddInstances(const UStaticMesh* InStaticMesh, const UMaterialInterface* InMaterial, const TArray<FTransform>& InTransforms, const TArray<FCustomPrimitiveData>& InCustomPrimitiveData);
 
-	UE_DEPRECATED(4.27, "Use AddInstances() instead. ")
-	ENGINE_API void SetupImposters(const UMaterialInterface* InImposterMaterial, UStaticMesh* InStaticMesh, const TArray<FTransform>& InTransforms);
-
 	/** Clear all instances stored in this LODActor. */
 	ENGINE_API void ClearInstances();
 
@@ -268,12 +265,7 @@ public:
 	ENGINE_API virtual bool Modify(bool bAlwaysMarkDirty = true) override;
 	ENGINE_API virtual void PreEditChange(FProperty* PropertyThatWillChange) override;
 	ENGINE_API virtual void PostEditChangeProperty(FPropertyChangedEvent& PropertyChangedEvent) override;
-
-	PRAGMA_DISABLE_DEPRECATION_WARNINGS // Suppress compiler warning on override of deprecated function
-	UE_DEPRECATED(5.0, "Use version that takes FObjectPreSaveContext instead.")
-	ENGINE_API virtual void PreSave(const class ITargetPlatform* TargetPlatform) override;
-	ENGINE_API PRAGMA_ENABLE_DEPRECATION_WARNINGS
-	virtual void PreSave(FObjectPreSaveContext ObjectSaveContext) override;
+	ENGINE_API virtual void PreSave(FObjectPreSaveContext ObjectSaveContext) override;
 #endif // WITH_EDITOR	
 
 	
@@ -324,9 +316,6 @@ public:
 	/** Returns instanced static mesh components that represents the imposters found in this LODActor. */
 	TArray<UInstancedStaticMeshComponent*> GetInstancedStaticMeshComponents() const { TArray<typename TDecay<decltype(InstancedStaticMeshComponents[FHLODInstancingKey()])>::Type> ISM; InstancedStaticMeshComponents.GenerateValueArray(ISM); return ISM; }
 
-	UE_DEPRECATED(4.27, "Use GetInstancedStaticMeshComponents() instead. ")
-	TArray<UInstancedStaticMeshComponent*> GetImpostersStaticMeshComponents() const { return GetInstancedStaticMeshComponents(); }
-
 	/** Returns an array of distances that are used to override individual LOD actors min draw distances. */
 	static ENGINE_API const TArray<float>& GetHLODDistanceOverride();
 
diff --git a/Engine/Source/Runtime/Engine/Private/HLODProxy.cpp b/Engine/Source/Runtime/Engine/Private/HLODProxy.cpp
--- a/Engine/Source/Runtime/Engine/Private/HLODProxy.cpp
+++ b/Engine/Source/Runtime/Engine/Private/HLODProxy.cpp
@@ -163,13 +163,6 @@ void UHLODProxy::DeletePackage()
 	ObjectTools::DeleteObjectsUnchecked({ Package });
 }
 
-void UHLODProxy::PreSave(const class ITargetPlatform* TargetPlatform)
-{
-	PRAGMA_DISABLE_DEPRECATION_WARNINGS;
-	Super::PreSave(TargetPlatform);
-	PRAGMA_ENABLE_DEPRECATION_WARNINGS;
-}
-
 void UHLODProxy::PreSave(FObjectPreSaveContext ObjectSaveContext)
 {
 	Super::PreSave(ObjectSaveContext);
diff --git a/Engine/Source/Runtime/Engine/Private/LODActor.cpp b/Engine/Source/Runtime/Engine/Private/LODActor.cpp
--- a/Engine/Source/Runtime/Engine/Private/LODActor.cpp
+++ b/Engine/Source/Runtime/Engine/Private/LODActor.cpp
@@ -1082,11 +1082,6 @@ void ALODActor::AddInstances(const UStaticMesh* InStaticMesh, const UMaterialInt
 	AddInstances(InStaticMesh, InMaterial, InTransforms, {});
 }
 
-void ALODActor::SetupImposters(const UMaterialInterface* InImposterMaterial, UStaticMesh* InStaticMesh, const TArray<FTransform>& InTransforms)
-{
-	AddInstances(InStaticMesh, InImposterMaterial, InTransforms);
-}
-
 void ALODActor::UpdateSubActorLODParents()
 {
 	for (AActor* Actor : SubActors)
@@ -1419,13 +1414,6 @@ void ALODActor::Serialize(FArchive& Ar)
 
 #if WITH_EDITOR
 
-void ALODActor::PreSave(const class ITargetPlatform* TargetPlatform)
-{
-	PRAGMA_DISABLE_DEPRECATION_WARNINGS;
-	Super::PreSave(TargetPlatform);
-	PRAGMA_ENABLE_DEPRECATION_WARNINGS;
-}
-
 void ALODActor::PreSave(FObjectPreSaveContext ObjectSaveContext)
 {
 	Super::PreSave(ObjectSaveContext);
