Materials: Fixes MaterialTextureScale generating invalid data for materials with world position offsets
* Disables WPO when rendering in RenderMaterialTexCoordScales.
* Fix wrong UV density results from alpha channel in  FMaterialUtilities::ExportMaterialUVDensities() by clearing with FColor::Transparent instead of FColor::Black
* Restore bNullifyWorldSpacePosition logic from CL 2955403 lost in merge commit CL 3019494

diff --git a/Engine/Source/Developer/MaterialUtilities/Private/MeshRendering.cpp b/Engine/Source/Developer/MaterialUtilities/Private/MeshRendering.cpp
--- a/Engine/Source/Developer/MaterialUtilities/Private/MeshRendering.cpp
+++ b/Engine/Source/Developer/MaterialUtilities/Private/MeshRendering.cpp
@@ -27,6 +27,7 @@
 #include "CanvasRender.h"
 #include "LocalVertexFactory.h"
 #include "Rendering/SkeletalMeshLODRenderData.h"
+#include "PrimitiveUniformShaderParametersBuilder.h"
 #include "MeshPassProcessor.h"
 #include "RHIStaticStates.h"
 #include "RendererInterface.h"
@@ -511,7 +512,19 @@ public:
 
 		FMeshBatch& MeshElement = *RenderContext.Alloc<FMeshBatch>();
 		FMeshBuilderOneFrameResources& OneFrameResource = *RenderContext.Alloc<FMeshBuilderOneFrameResources>();
-		DynamicMeshBuilder.GetMeshElement(FMatrix::Identity, Data.MaterialRenderProxy, SDPG_Foreground, true, false, 0, OneFrameResource, MeshElement);
+
+		// Provide custom primitive params to prevent world position offset from moving the mesh
+		FPrimitiveUniformShaderParameters PrimitiveParams = FPrimitiveUniformShaderParametersBuilder{}
+			.Defaults()
+				.LocalToWorld(FMatrix::Identity)
+				.ActorWorldPosition(FMatrix::Identity.GetOrigin())
+				.WorldBounds(FBoxSphereBounds(EForceInit::ForceInit))
+				.LocalBounds(FBoxSphereBounds(EForceInit::ForceInit))
+				.ReceivesDecals(false)
+				.OutputVelocity(true)
+				.EvaluateWorldPositionOffset(false)
+			.Build();
+		DynamicMeshBuilder.GetMeshElement(PrimitiveParams, Data.MaterialRenderProxy, SDPG_Foreground, true, 0, OneFrameResource, MeshElement);
 
 		check(OneFrameResource.IsValidForRendering());
 
@@ -753,11 +766,11 @@ bool FMeshRenderer::RenderMaterialTexCoordScales(struct FMaterialMergeData& InMa
 	check(IsInGameThread());
 	check(InRenderTarget);
 
-	// Create a canvas for the render target and clear it to black
+	// Create a canvas for the render target and clear it to transparent
 	FTextureRenderTargetResource* RTResource = InRenderTarget->GameThread_GetRenderTargetResource();
 	FCanvas Canvas(RTResource, NULL, FGameTime::GetTimeSinceAppStart(), GMaxRHIFeatureLevel);
 	const FRenderTarget* CanvasRenderTarget = Canvas.GetRenderTarget();
-	Canvas.Clear(FLinearColor::Black);
+	Canvas.Clear(FLinearColor::Transparent);
 
 	// Set show flag view mode to output tex coord scale
 	FEngineShowFlags ShowFlags(ESFIM_Game);
diff --git a/Engine/Source/Runtime/Engine/Private/SceneView.cpp b/Engine/Source/Runtime/Engine/Private/SceneView.cpp
--- a/Engine/Source/Runtime/Engine/Private/SceneView.cpp
+++ b/Engine/Source/Runtime/Engine/Private/SceneView.cpp
@@ -2464,6 +2464,18 @@ void FSceneView::SetupViewRectUniformBufferParameters(FViewUniformShaderParamete
 
 	ViewUniformShaderParameters.MotionBlurNormalizedToPixel = FinalPostProcessSettings.MotionBlurMax * EffectiveViewRect.Width() / 100.0f;
 
+	#if WITH_EDITOR
+		if (Family->bNullifyWorldSpacePosition)
+		{
+			// Forces world space position to 0 and view vector to up.
+			ViewUniformShaderParameters.SVPositionToTranslatedWorld = FMatrix44f(
+				FMatrix(FPlane(0, 0, 0, 0),
+					FPlane(0, 0, 0, 0),
+					FPlane(0, 0, 0, 0),
+					FPlane(0, 0, 1, 1)));
+		}
+		else
+	#endif 
 	{
 		// setup a matrix to transform float4(SvPosition.xyz,1) directly to TranslatedWorld (quality, performance as we don't need to convert or use interpolator)
 
