Use a view state when doing a material capture.
Without this there is no frame counter in the rendered view.
When we are performing virtual texture warm up frames we need the frame counter since it is used as the input to the calculation that jitters the virtual texture feedback sample locations.

diff --git a/Engine/Source/Developer/MaterialUtilities/Private/MaterialUtilities.cpp b/Engine/Source/Developer/MaterialUtilities/Private/MaterialUtilities.cpp
--- a/Engine/Source/Developer/MaterialUtilities/Private/MaterialUtilities.cpp
+++ b/Engine/Source/Developer/MaterialUtilities/Private/MaterialUtilities.cpp
@@ -792,16 +792,23 @@ private:
  * the CVarMaterialUtilitiesWarmupFrames CVar. This is needed to ensure various rendering systems are primed properly before capturing
  * the scene.
  */
-static void PerformSceneRender(FCanvas& Canvas, FSceneViewFamily& ViewFamily, bool bWithWarmup)
+static void PerformSceneRender(FCanvas& Canvas, FSceneViewFamily& ViewFamily, FSceneView& View, bool bWithWarmup)
 {
 	if (bWithWarmup)
 	{
 		for (int32 i = 0; i < CVarMaterialUtilitiesWarmupFrames.GetValueOnGameThread(); i++)
 		{
+			// Increment frame index when warming up virtual textures.
+			// This is necessary to jitter the virtual texture feedback location.
+			View.OverrideFrameIndexValue = i;
+
 			GetRendererModule().BeginRenderingViewFamily(&Canvas, &ViewFamily);
 		}
 	}
 
+	// Do final capture with deterministic frame index 0.
+	View.OverrideFrameIndexValue = 0;
+
 	GetRendererModule().BeginRenderingViewFamily(&Canvas, &ViewFamily);
 }
 
@@ -837,12 +844,16 @@ static void RenderSceneToTexture(
 	ViewFamily.EngineShowFlags.SetTonemapper(false);
 	ViewFamily.EngineShowFlags.SetScreenPercentage(false);
 
+	FSceneViewStateReference ViewState;
+	ViewState.Allocate(Scene->GetFeatureLevel());
+
 	FSceneViewInitOptions ViewInitOptions;
 	ViewInitOptions.SetViewRectangle(FIntRect(0, 0, TargetSize.X, TargetSize.Y));
 	ViewInitOptions.ViewFamily = &ViewFamily;
 	ViewInitOptions.ViewOrigin = ViewOrigin;
 	ViewInitOptions.ViewRotationMatrix = ViewRotationMatrix;
 	ViewInitOptions.ProjectionMatrix = ProjectionMatrix;
+	ViewInitOptions.SceneViewStateInterface = ViewState.GetReference();
 
 	// If no "show only" primitives are provided, we must pass an unset TOptional - otherwise an empty set will mean no primitive should be visible.
 	ViewInitOptions.ShowOnlyPrimitives = !ShowOnlyPrimitives.IsEmpty() ? TOptional<TSet<FPrimitiveComponentId>>(ShowOnlyPrimitives) : TOptional<TSet<FPrimitiveComponentId>>();
@@ -865,7 +876,7 @@ static void RenderSceneToTexture(
 		Extension->SetupView(ViewFamily, *NewView);
 	}
 
-	PerformSceneRender(Canvas, ViewFamily, bPerformWarmpup);
+	PerformSceneRender(Canvas, ViewFamily, *NewView, bPerformWarmpup);
 
 	// Copy the contents of the remote texture to system memory
 	OutSamples.SetNumUninitialized(TargetSize.X*TargetSize.Y);
