ImageUtils::ImageResize followup fixing callers incorrectly using the bool Linear argument

diff --git a/Engine/Plugins/Enterprise/VariantManagerContent/Source/VariantManagerContent/Private/ThumbnailGenerator.cpp b/Engine/Plugins/Enterprise/VariantManagerContent/Source/VariantManagerContent/Private/ThumbnailGenerator.cpp
--- a/Engine/Plugins/Enterprise/VariantManagerContent/Source/VariantManagerContent/Private/ThumbnailGenerator.cpp
+++ b/Engine/Plugins/Enterprise/VariantManagerContent/Source/VariantManagerContent/Private/ThumbnailGenerator.cpp
@@ -371,7 +371,7 @@ UTexture2D* ThumbnailGenerator::GenerateThumbnailFromObjectThumbnail(UObject* Ob
             TArray<FColor> ResizedColors;
             ResizedColors.SetNum(VARIANT_MANAGER_THUMBNAIL_SIZE * VARIANT_MANAGER_THUMBNAIL_SIZE);
 
-			const bool bLinearSpace = false;
+			const bool bLinearSpace = true;
 			const bool bForceOpaqueOutput = false;
 			FImageUtils::ImageResize( Width, Height, OldColors, VARIANT_MANAGER_THUMBNAIL_SIZE, VARIANT_MANAGER_THUMBNAIL_SIZE, ResizedColors, bLinearSpace, bForceOpaqueOutput );
 
diff --git a/Engine/Source/Developer/MaterialUtilities/Private/MaterialUtilities.cpp b/Engine/Source/Developer/MaterialUtilities/Private/MaterialUtilities.cpp
--- a/Engine/Source/Developer/MaterialUtilities/Private/MaterialUtilities.cpp
+++ b/Engine/Source/Developer/MaterialUtilities/Private/MaterialUtilities.cpp
@@ -1912,8 +1912,7 @@ void FMaterialUtilities::ResizeFlattenMaterial(FFlattenMaterial& InFlattenMateri
 		if (InFlattenMaterial.GetPropertySize(EFlattenMaterialProperties::Diffuse).X != DiffuseSizeX)
 		{
 			TArray<FColor> NewSamples;
-			// @todo : last arg should be true not false
-			FImageUtils::ImageResize(InFlattenMaterial.GetPropertySize(EFlattenMaterialProperties::Diffuse).X, InFlattenMaterial.GetPropertySize(EFlattenMaterialProperties::Diffuse).Y, InFlattenMaterial.GetPropertySamples(EFlattenMaterialProperties::Diffuse), DiffuseSizeX, DiffuseSizeX, NewSamples, false);
+			FImageUtils::ImageResize(InFlattenMaterial.GetPropertySize(EFlattenMaterialProperties::Diffuse).X, InFlattenMaterial.GetPropertySize(EFlattenMaterialProperties::Diffuse).Y, InFlattenMaterial.GetPropertySamples(EFlattenMaterialProperties::Diffuse), DiffuseSizeX, DiffuseSizeX, NewSamples, true);
 			InFlattenMaterial.GetPropertySamples(EFlattenMaterialProperties::Diffuse).Reset(NewSamples.Num());
 			InFlattenMaterial.GetPropertySamples(EFlattenMaterialProperties::Diffuse).Append(NewSamples);
 			InFlattenMaterial.SetPropertySize(EFlattenMaterialProperties::Diffuse, FIntPoint(DiffuseSizeX, DiffuseSizeX));
@@ -1949,8 +1948,7 @@ void FMaterialUtilities::ResizeFlattenMaterial(FFlattenMaterial& InFlattenMateri
 		if (InFlattenMaterial.GetPropertySamples(EFlattenMaterialProperties::Normal).Num() && MaterialSettings.bNormalMap && InFlattenMaterial.GetPropertySize(EFlattenMaterialProperties::Normal).X != NormalSizeX)
 		{
 			TArray<FColor> NewSamples;
-			// @todo : last arg should be false not true
-			FImageUtils::ImageResize(InFlattenMaterial.GetPropertySize(EFlattenMaterialProperties::Normal).X, InFlattenMaterial.GetPropertySize(EFlattenMaterialProperties::Normal).Y, InFlattenMaterial.GetPropertySamples(EFlattenMaterialProperties::Normal), NormalSizeX, NormalSizeX, NewSamples, true);
+			FImageUtils::ImageResize(InFlattenMaterial.GetPropertySize(EFlattenMaterialProperties::Normal).X, InFlattenMaterial.GetPropertySize(EFlattenMaterialProperties::Normal).Y, InFlattenMaterial.GetPropertySamples(EFlattenMaterialProperties::Normal), NormalSizeX, NormalSizeX, NewSamples, false);
 			InFlattenMaterial.GetPropertySamples(EFlattenMaterialProperties::Normal).Reset(NewSamples.Num());
 			InFlattenMaterial.GetPropertySamples(EFlattenMaterialProperties::Normal).Append(NewSamples);
 			InFlattenMaterial.SetPropertySize(EFlattenMaterialProperties::Normal, FIntPoint(NormalSizeX, NormalSizeX));
@@ -1959,7 +1957,7 @@ void FMaterialUtilities::ResizeFlattenMaterial(FFlattenMaterial& InFlattenMateri
 		if (InFlattenMaterial.GetPropertySamples(EFlattenMaterialProperties::Emissive).Num() && MaterialSettings.bEmissiveMap && InFlattenMaterial.GetPropertySize(EFlattenMaterialProperties::Emissive).X != PropertiesSizeX)
 		{
 			TArray<FColor> NewSamples;
-			FImageUtils::ImageResize(InFlattenMaterial.GetPropertySize(EFlattenMaterialProperties::Emissive).X, InFlattenMaterial.GetPropertySize(EFlattenMaterialProperties::Emissive).Y, InFlattenMaterial.GetPropertySamples(EFlattenMaterialProperties::Emissive), PropertiesSizeX, PropertiesSizeX, NewSamples, false);
+			FImageUtils::ImageResize(InFlattenMaterial.GetPropertySize(EFlattenMaterialProperties::Emissive).X, InFlattenMaterial.GetPropertySize(EFlattenMaterialProperties::Emissive).Y, InFlattenMaterial.GetPropertySamples(EFlattenMaterialProperties::Emissive), PropertiesSizeX, PropertiesSizeX, NewSamples, true);
 			InFlattenMaterial.GetPropertySamples(EFlattenMaterialProperties::Emissive).Reset(NewSamples.Num());
 			InFlattenMaterial.GetPropertySamples(EFlattenMaterialProperties::Emissive).Append(NewSamples);
 			InFlattenMaterial.SetPropertySize(EFlattenMaterialProperties::Emissive, FIntPoint(PropertiesSizeX, PropertiesSizeX));
@@ -1988,7 +1986,7 @@ void FMaterialUtilities::ResizeFlattenMaterial(FFlattenMaterial& InFlattenMateri
 		if (InFlattenMaterial.GetPropertySize(EFlattenMaterialProperties::Diffuse) != MaterialSettings.DiffuseTextureSize)
 		{
 			TArray<FColor> NewSamples;
-			FImageUtils::ImageResize(InFlattenMaterial.GetPropertySize(EFlattenMaterialProperties::Diffuse).X, InFlattenMaterial.GetPropertySize(EFlattenMaterialProperties::Diffuse).Y, InFlattenMaterial.GetPropertySamples(EFlattenMaterialProperties::Diffuse), MaterialSettings.DiffuseTextureSize.X, MaterialSettings.DiffuseTextureSize.Y, NewSamples, false);
+			FImageUtils::ImageResize(InFlattenMaterial.GetPropertySize(EFlattenMaterialProperties::Diffuse).X, InFlattenMaterial.GetPropertySize(EFlattenMaterialProperties::Diffuse).Y, InFlattenMaterial.GetPropertySamples(EFlattenMaterialProperties::Diffuse), MaterialSettings.DiffuseTextureSize.X, MaterialSettings.DiffuseTextureSize.Y, NewSamples, true);
 			InFlattenMaterial.GetPropertySamples(EFlattenMaterialProperties::Diffuse).Reset(NewSamples.Num());
 			InFlattenMaterial.GetPropertySamples(EFlattenMaterialProperties::Diffuse).Append(NewSamples);
 			InFlattenMaterial.SetPropertySize(EFlattenMaterialProperties::Diffuse, MaterialSettings.DiffuseTextureSize);
@@ -2024,7 +2022,7 @@ void FMaterialUtilities::ResizeFlattenMaterial(FFlattenMaterial& InFlattenMateri
 		if (InFlattenMaterial.GetPropertySamples(EFlattenMaterialProperties::Normal).Num() && MaterialSettings.bNormalMap && InFlattenMaterial.GetPropertySize(EFlattenMaterialProperties::Normal) != MaterialSettings.NormalTextureSize)
 		{
 			TArray<FColor> NewSamples;
-			FImageUtils::ImageResize(InFlattenMaterial.GetPropertySize(EFlattenMaterialProperties::Normal).X, InFlattenMaterial.GetPropertySize(EFlattenMaterialProperties::Normal).Y, InFlattenMaterial.GetPropertySamples(EFlattenMaterialProperties::Normal), MaterialSettings.NormalTextureSize.X, MaterialSettings.NormalTextureSize.Y, NewSamples, true);
+			FImageUtils::ImageResize(InFlattenMaterial.GetPropertySize(EFlattenMaterialProperties::Normal).X, InFlattenMaterial.GetPropertySize(EFlattenMaterialProperties::Normal).Y, InFlattenMaterial.GetPropertySamples(EFlattenMaterialProperties::Normal), MaterialSettings.NormalTextureSize.X, MaterialSettings.NormalTextureSize.Y, NewSamples, false);
 			InFlattenMaterial.GetPropertySamples(EFlattenMaterialProperties::Normal).Reset(NewSamples.Num());
 			InFlattenMaterial.GetPropertySamples(EFlattenMaterialProperties::Normal).Append(NewSamples);
 			InFlattenMaterial.SetPropertySize(EFlattenMaterialProperties::Normal, MaterialSettings.NormalTextureSize);
@@ -2033,7 +2031,7 @@ void FMaterialUtilities::ResizeFlattenMaterial(FFlattenMaterial& InFlattenMateri
 		if (InFlattenMaterial.GetPropertySamples(EFlattenMaterialProperties::Emissive).Num() && MaterialSettings.bEmissiveMap && InFlattenMaterial.GetPropertySize(EFlattenMaterialProperties::Emissive) != MaterialSettings.EmissiveTextureSize)
 		{
 			TArray<FColor> NewSamples;
-			FImageUtils::ImageResize(InFlattenMaterial.GetPropertySize(EFlattenMaterialProperties::Emissive).X, InFlattenMaterial.GetPropertySize(EFlattenMaterialProperties::Emissive).Y, InFlattenMaterial.GetPropertySamples(EFlattenMaterialProperties::Emissive), MaterialSettings.EmissiveTextureSize.X, MaterialSettings.EmissiveTextureSize.Y, NewSamples, false);
+			FImageUtils::ImageResize(InFlattenMaterial.GetPropertySize(EFlattenMaterialProperties::Emissive).X, InFlattenMaterial.GetPropertySize(EFlattenMaterialProperties::Emissive).Y, InFlattenMaterial.GetPropertySamples(EFlattenMaterialProperties::Emissive), MaterialSettings.EmissiveTextureSize.X, MaterialSettings.EmissiveTextureSize.Y, NewSamples, true);
 			InFlattenMaterial.GetPropertySamples(EFlattenMaterialProperties::Emissive).Reset(NewSamples.Num());
 			InFlattenMaterial.GetPropertySamples(EFlattenMaterialProperties::Emissive).Append(NewSamples);
 			InFlattenMaterial.SetPropertySize(EFlattenMaterialProperties::Emissive, MaterialSettings.EmissiveTextureSize);
diff --git a/Engine/Source/Developer/MeshMergeUtilities/Private/MeshMergeUtilities.cpp b/Engine/Source/Developer/MeshMergeUtilities/Private/MeshMergeUtilities.cpp
--- a/Engine/Source/Developer/MeshMergeUtilities/Private/MeshMergeUtilities.cpp
+++ b/Engine/Source/Developer/MeshMergeUtilities/Private/MeshMergeUtilities.cpp
@@ -883,7 +883,8 @@ FIntPoint FMeshMergeUtilities::ConditionalImageResize(const FIntPoint& SrcSize,
 		TArray<FColor> OutImage;
 		if (NumDesiredSamples > 0)
 		{
-			FImageUtils::ImageResize(SrcSize.X, SrcSize.Y, InOutImage, DesiredSize.X, DesiredSize.Y, OutImage, bLinearSpace);
+			const bool bOpaqueAlpha = true; // ??
+			FImageUtils::ImageResize(SrcSize.X, SrcSize.Y, InOutImage, DesiredSize.X, DesiredSize.Y, OutImage, bLinearSpace, bOpaqueAlpha);
 		}
 		Exchange(InOutImage, OutImage);
 		return DesiredSize;
diff --git a/Engine/Source/Editor/UnrealEd/Private/ThumbnailExternalCache.cpp b/Engine/Source/Editor/UnrealEd/Private/ThumbnailExternalCache.cpp
--- a/Engine/Source/Editor/UnrealEd/Private/ThumbnailExternalCache.cpp
+++ b/Engine/Source/Editor/UnrealEd/Private/ThumbnailExternalCache.cpp
@@ -35,7 +35,7 @@ namespace ThumbnailExternalCache
 		Thumbnail.GetUncompressedImageData();
 		TArray<uint8>& UncompressedImageData = Thumbnail.AccessImageData();
 
-		const bool bLinearSpace = false;
+		const bool bLinearSpace = true;
 		const bool bForceOpaqueOutput = false;
 		const TArrayView<FColor> SrcDataView(reinterpret_cast<FColor*>(UncompressedImageData.GetData()), UncompressedImageData.Num() / sizeof(FColor));
 		const TArrayView<FColor> DestDataView(reinterpret_cast<FColor*>(DestData.GetData()), DestData.Num() / sizeof(FColor));
diff --git a/Engine/Source/Runtime/Engine/Private/ImageUtils.cpp b/Engine/Source/Runtime/Engine/Private/ImageUtils.cpp
--- a/Engine/Source/Runtime/Engine/Private/ImageUtils.cpp
+++ b/Engine/Source/Runtime/Engine/Private/ImageUtils.cpp
@@ -600,30 +600,29 @@ bool FImageUtils::GetRenderTargetImage(UTextureRenderTarget* TexRT, FImage & Out
 
 
 // prefer FImageCore::ResizeImage
-void FImageUtils::ImageResize(int32 SrcWidth, int32 SrcHeight, const TArray<FColor> &SrcData, int32 DstWidth, int32 DstHeight, TArray<FColor> &DstData, bool bLinearSpace, bool bForceOpaqueOutput)
+void FImageUtils::ImageResize(int32 SrcWidth, int32 SrcHeight, const TArray<FColor> &SrcData, int32 DstWidth, int32 DstHeight, TArray<FColor> &DstData, bool bResizeSRGBinLinearSpace, bool bForceOpaqueOutput)
 {
-	DstData.Empty(DstWidth*DstHeight);
-	DstData.AddZeroed(DstWidth*DstHeight);
+	int64 Num = (int64)DstWidth*DstHeight;
+	DstData.Empty(Num);
+	DstData.AddZeroed(Num);
 
-	ImageResize(SrcWidth, SrcHeight, TArrayView<const FColor>(SrcData), DstWidth, DstHeight, TArrayView<FColor>(DstData), bLinearSpace, bForceOpaqueOutput);
+	ImageResize(SrcWidth, SrcHeight, TArrayView<const FColor>(SrcData), DstWidth, DstHeight, TArrayView<FColor>(DstData), bResizeSRGBinLinearSpace, bForceOpaqueOutput);
 }
 
 // prefer FImageCore::ResizeImage
-void FImageUtils::ImageResize(int32 SrcWidth, int32 SrcHeight, const TArrayView<const FColor> &SrcData, int32 DstWidth, int32 DstHeight, const TArrayView<FColor> &DstData, bool bLinearSpace, bool bForceOpaqueOutput)
+void FImageUtils::ImageResize(int32 SrcWidth, int32 SrcHeight, const TArrayView<const FColor> &SrcData, int32 DstWidth, int32 DstHeight, const TArrayView<FColor> &DstData, bool bResizeSRGBinLinearSpace, bool bForceOpaqueOutput)
 {
-	check(SrcData.Num() >= SrcWidth * SrcHeight);
-	check(DstData.Num() >= DstWidth * DstHeight);
+	check(SrcData.Num() >= (int64)SrcWidth * SrcHeight);
+	check(DstData.Num() >= (int64)DstWidth * DstHeight);
 	
 	if ( SrcData.Num() == 0 || DstData.Num() == 0 )
 	{
 		return;
 	}
 
-	// bLinearSpace means convert SRGB colors to Linear before resizing
-	//
-	//	(there are a few callers using bLinearSpace=false , most of which seem to be unintentional/bugs
-	//	 MaterialUtilities.cpp attempts to use it right but does NOT)
-	EGammaSpace Gamma = bLinearSpace ? EGammaSpace::sRGB : EGammaSpace::Linear; 
+	// bResizeSRGBinLinearSpace means convert SRGB colors to Linear before resizing
+	//	if it's false, we tell resize that our U8's are linear so it does no transforms
+	EGammaSpace Gamma = bResizeSRGBinLinearSpace ? EGammaSpace::sRGB : EGammaSpace::Linear; 
 
 	FImageView SrcImage(&SrcData[0],SrcWidth,SrcHeight,Gamma);
 	FImageView DstImage(&DstData[0],DstWidth,DstHeight,Gamma);
@@ -639,8 +638,9 @@ void FImageUtils::ImageResize(int32 SrcWidth, int32 SrcHeight, const TArrayView<
 // prefer FImageCore::ResizeImage
 void FImageUtils::ImageResize(int32 SrcWidth, int32 SrcHeight, const TArray64<FLinearColor>& SrcData, int32 DstWidth, int32 DstHeight, TArray64<FLinearColor>& DstData)
 {
-	DstData.Empty(DstWidth * DstHeight);
-	DstData.AddZeroed(DstWidth * DstHeight);
+	int64 Num = (int64)DstWidth*DstHeight;
+	DstData.Empty(Num);
+	DstData.AddZeroed(Num);
 
 	ImageResize(SrcWidth, SrcHeight, TArrayView64<const FLinearColor>(SrcData), DstWidth, DstHeight, TArrayView64<FLinearColor>(DstData));
 }
@@ -648,8 +648,8 @@ void FImageUtils::ImageResize(int32 SrcWidth, int32 SrcHeight, const TArray64<FL
 // prefer FImageCore::ResizeImage
 void FImageUtils::ImageResize(int32 SrcWidth, int32 SrcHeight, const TArrayView64<const FLinearColor>& SrcData, int32 DstWidth, int32 DstHeight, const TArrayView64<FLinearColor>& DstData)
 {
-	check(SrcData.Num() >= SrcWidth * SrcHeight);
-	check(DstData.Num() >= DstWidth * DstHeight);
+	check(SrcData.Num() >= (int64)SrcWidth * SrcHeight);
+	check(DstData.Num() >= (int64)DstWidth * DstHeight);
 	
 	if ( SrcData.Num() == 0 || DstData.Num() == 0 )
 	{
@@ -838,7 +838,7 @@ void FImageUtils::CropAndScaleImage( int32 SrcWidth, int32 SrcHeight, int32 Desi
 	DstData.AddUninitialized(DesiredWidth*DesiredHeight);
 
 	// Resize the image
-	FImageUtils::ImageResize( MaxWidth, MaxHeight, CroppedData, DesiredWidth, DesiredHeight, DstData, true );
+	FImageUtils::ImageResize( MaxWidth, MaxHeight, CroppedData, DesiredWidth, DesiredHeight, DstData, true, true );
 }
 
 void FImageUtils::ThumbnailCompressImageArray( int32 ImageWidth, int32 ImageHeight, const TArray<FColor> &SrcData, TArray<uint8> &DstData )
diff --git a/Engine/Source/Runtime/Engine/Private/ProfilingDebugging/TraceScreenshot.cpp b/Engine/Source/Runtime/Engine/Private/ProfilingDebugging/TraceScreenshot.cpp
--- a/Engine/Source/Runtime/Engine/Private/ProfilingDebugging/TraceScreenshot.cpp
+++ b/Engine/Source/Runtime/Engine/Private/ProfilingDebugging/TraceScreenshot.cpp
@@ -98,7 +98,8 @@ void FTraceScreenshot::RequestScreenshot(FString Name, bool bShowUI, const FLogC
 
 void ImageUtilsImageResize(int32 SrcWidth, int32 SrcHeight, const TArray<FColor>& SrcData, int32 DstWidth, int32 DstHeight, TArray<FColor>& DstData, bool bLinearSpace)
 {
-	FImageUtils::ImageResize(SrcWidth, SrcHeight, SrcData, DstWidth, DstHeight, DstData, bLinearSpace);
+	const bool bOpaqueAlpha = true;
+	FImageUtils::ImageResize(SrcWidth, SrcHeight, SrcData, DstWidth, DstHeight, DstData, bLinearSpace, bOpaqueAlpha);
 }
 
 void ImageUtilsCompressImageArrayWrapper(int32 ImageWidth, int32 ImageHeight, const TArrayView64<const FColor>& SrcData, TArray64<uint8>& DstData)
@@ -108,6 +109,7 @@ void ImageUtilsCompressImageArrayWrapper(int32 ImageWidth, int32 ImageHeight, co
 
 void ImageUtilsImageResize(int32 SrcWidth, int32 SrcHeight, const TArray64<FLinearColor>& SrcData, int32 DstWidth, int32 DstHeight, TArray64<FLinearColor>& DstData, bool /*bUnused*/)
 {
+	// note: does NOT opaque alpha
 	FImageUtils::ImageResize(SrcWidth, SrcHeight, SrcData, DstWidth, DstHeight, DstData);
 }
 
@@ -154,7 +156,7 @@ void TraceScreenshotInternal(int32 InSizeX, int32 InSizeY, const TArray<FColorTy
 
 				FResizedColorArrayType ResizedImage;
 				ResizedImage.SetNum(ResizedX * ResizedY);
-				ImageUtilsImageResize(InSizeX, InSizeY, *ImageCopy, ResizedX, ResizedY, ResizedImage, false);
+				ImageUtilsImageResize(InSizeX, InSizeY, *ImageCopy, ResizedX, ResizedY, ResizedImage, true);
 
 				ImageUtilsCompressImageArrayWrapper(ResizedX, ResizedY, TArrayView64<const FColorType>(ResizedImage.GetData(), ResizedImage.Num()), CompressedBitmap);
 				TRACE_SCREENSHOT(*ScreenshotName, Cycles, ResizedX, ResizedY, CompressedBitmap);
diff --git a/Engine/Source/Runtime/Engine/Public/ImageUtils.h b/Engine/Source/Runtime/Engine/Public/ImageUtils.h
--- a/Engine/Source/Runtime/Engine/Public/ImageUtils.h
+++ b/Engine/Source/Runtime/Engine/Public/ImageUtils.h
@@ -184,10 +184,12 @@ public:
 	 * @param DstWidth				Destination image width.
 	 * @param DstHeight				Destination image height.
 	 * @param DstData				Destination image data. (will be resized)
-	 * @param bLinearSpace			Convert linear to sRGB for gamma-correct resize.  If this is false source U8s are treated as linear.  Should usually be true.
+	 * @param bResizeSRGBinLinearSpace			Convert sRGB to linear for gamma-correct resize.  Should be true if the Src is SRGB diffuse color; should be false for masks and normal maps.
 	 * @param bForceOpaqueOutput	Always output 255 for the alpha channel.  BEWARE defaults to true so alpha is killed by default.
+	 *
+	 * BEWARE: the FColor variant of this function fills with opaque alpha by default.  The FLinearColor variant does not
 	 */
-	ENGINE_API static void ImageResize(int32 SrcWidth, int32 SrcHeight, const TArray<FColor> &SrcData,  int32 DstWidth, int32 DstHeight, TArray<FColor> &DstData, bool bLinearSpace, bool bForceOpaqueOutput=true);
+	ENGINE_API static void ImageResize(int32 SrcWidth, int32 SrcHeight, const TArray<FColor> &SrcData,  int32 DstWidth, int32 DstHeight, TArray<FColor> &DstData, bool bResizeSRGBinLinearSpace, bool bForceOpaqueOutput=true);
 
 	/**
 	 * Resizes the given image using FImageCore::ResizeImage.
@@ -202,10 +204,12 @@ public:
 	 * @param DstWidth				Destination image width.
 	 * @param DstHeight				Destination image height.
 	 * @param DstData				Destination image data. (must already be sized to DstWidth*DstHeight)
-	 * @param bLinearSpace			Convert linear to sRGB for gamma-correct resize.  If this is false source U8s are treated as linear.  Should usually be true.
+ 	 * @param bResizeSRGBinLinearSpace			Convert sRGB to linear for gamma-correct resize.  Should be true if the Src is SRGB diffuse color; should be false for masks and normal maps.
 	 * @param bForceOpaqueOutput	Always output 255 for the alpha channel.  BEWARE defaults to true so alpha is killed by default.
+	 *
+	 * BEWARE: the FColor variant of this function fills with opaque alpha by default.  The FLinearColor variant does not
 	 */
-	ENGINE_API static void ImageResize(int32 SrcWidth, int32 SrcHeight, const TArrayView<const FColor> &SrcData, int32 DstWidth, int32 DstHeight, const TArrayView<FColor> &DstData, bool bLinearSpace, bool bForceOpaqueOutput=true);
+	ENGINE_API static void ImageResize(int32 SrcWidth, int32 SrcHeight, const TArrayView<const FColor> &SrcData, int32 DstWidth, int32 DstHeight, const TArrayView<FColor> &DstData, bool bResizeSRGBinLinearSpace, bool bForceOpaqueOutput=true);
 
 	/**
 	 * Resizes the given image using FImageCore::ResizeImage.
@@ -220,6 +224,8 @@ public:
 	 * @param DstWidth	Destination image width.
 	 * @param DstHeight Destination image height.
 	 * @param DstData	Destination image data. (will be resized)
+	 *
+	 * BEWARE: the FColor variant of this function fills with opaque alpha by default.  The FLinearColor variant does not
 	 */
 	ENGINE_API static void ImageResize(int32 SrcWidth, int32 SrcHeight, const TArray64<FLinearColor>& SrcData, int32 DstWidth, int32 DstHeight, TArray64<FLinearColor>& DstData);
 
@@ -236,6 +242,8 @@ public:
 	 * @param DstWidth	Destination image width.
 	 * @param DstHeight Destination image height.
 	 * @param DstData	Destination image data. (must already be sized to DstWidth*DstHeight)
+	 *
+	 * BEWARE: the FColor variant of this function fills with opaque alpha by default.  The FLinearColor variant does not
 	 */
 	ENGINE_API static void ImageResize(int32 SrcWidth, int32 SrcHeight, const TArrayView64<const FLinearColor>& SrcData, int32 DstWidth, int32 DstHeight, const TArrayView64<FLinearColor>& DstData);
 
