Fixed material baking to account for render target linear / sRGB format.

#tests PC editor dx12. Bake Material button from mesh component.
#rnx

diff --git a/Engine/Source/Developer/MaterialBaking/Private/MaterialBakingModule.cpp b/Engine/Source/Developer/MaterialBaking/Private/MaterialBakingModule.cpp
--- a/Engine/Source/Developer/MaterialBaking/Private/MaterialBakingModule.cpp
+++ b/Engine/Source/Developer/MaterialBaking/Private/MaterialBakingModule.cpp
@@ -393,6 +393,11 @@ FBakeOutput ToBakeOutput(FBakeOutputEx& BakeOutputEx)
 		BakeOutput.HDRPropertyData.Add(PropertyDataPair.Key.Type, MoveTemp(PropertyDataPair.Value));
 	}
 
+	for (TPair<FMaterialPropertyEx, bool>& PropertyDataPair : BakeOutputEx.PropertyIsLinearColor)
+	{
+		BakeOutput.PropertyIsLinearColor.Add(PropertyDataPair.Key.Type, PropertyDataPair.Value);
+	}
+
 	return BakeOutput;
 }
 
@@ -549,6 +554,8 @@ protected:
 			{
 				BakeOutput.HDRPropertyData.Add(Property);
 			}
+
+			BakeOutput.PropertyIsLinearColor.Add(Property, MaterialBakingModule.IsLinearBake(Property));
 		}
 	}
 
diff --git a/Engine/Source/Developer/MaterialBaking/Public/MaterialBakingStructures.h b/Engine/Source/Developer/MaterialBaking/Public/MaterialBakingStructures.h
--- a/Engine/Source/Developer/MaterialBaking/Public/MaterialBakingStructures.h
+++ b/Engine/Source/Developer/MaterialBaking/Public/MaterialBakingStructures.h
@@ -166,6 +166,9 @@ struct FBakeOutput
 	/** Contains the resulting HDR texture data for baking out a material's property, may be empty */
 	TMap<EMaterialProperty, TArray<FFloat16Color>> HDRPropertyData;
 
+	/** Wether or not a baked texture data is using a linear color or sRGB encoding. */
+	TMap<EMaterialProperty, bool> PropertyIsLinearColor;
+
 	/** Scale used to allow having wide ranges of emissive values in the source materials, the final proxy material will use this value to scale the emissive texture's pixel values */
 	float EmissiveScale = 1.0f;
 };
@@ -182,6 +185,9 @@ struct FBakeOutputEx
 	/** Contains the resulting HDR texture data for baking out a material's property, may be empty */
 	TMap<FMaterialPropertyEx, TArray<FFloat16Color>> HDRPropertyData;
 
+	/** Wether or not a baked texture data is using a linear color or sRGB encoding. */
+	TMap<FMaterialPropertyEx, bool> PropertyIsLinearColor;
+
 	/** Scale used to allow having wide ranges of emissive values in the source materials, the final proxy material will use this value to scale the emissive texture's pixel values */
 	float EmissiveScale = 1.0f;
 };
diff --git a/Engine/Source/Developer/MaterialUtilities/Private/MaterialUtilities.cpp b/Engine/Source/Developer/MaterialUtilities/Private/MaterialUtilities.cpp
--- a/Engine/Source/Developer/MaterialUtilities/Private/MaterialUtilities.cpp
+++ b/Engine/Source/Developer/MaterialUtilities/Private/MaterialUtilities.cpp
@@ -171,11 +171,25 @@ UMaterialInterface* FMaterialUtilities::CreateProxyMaterialAndTextures(UPackage*
 
 			if (Property == MP_BaseColor || Property == MP_EmissiveColor)
 			{
-				Material->SetVectorParameterValueEditorOnly(ParameterInfo, FLinearColor::FromSRGBColor(ColorData[0]));
+				if(BakeOutput.PropertyIsLinearColor.Contains(Property) && BakeOutput.PropertyIsLinearColor.FindChecked(Property))
+				{
+					Material->SetVectorParameterValueEditorOnly(ParameterInfo, FLinearColor(float(ColorData[0].R) / 255.0f, float(ColorData[0].G) / 255.0f, float(ColorData[0].B) / 255.0f, float(ColorData[0].A) / 255.0f));
+				}
+				else
+				{
+					Material->SetVectorParameterValueEditorOnly(ParameterInfo, FLinearColor::FromSRGBColor(ColorData[0]));
+				}
 			}
 			else
 			{
-				Material->SetScalarParameterValueEditorOnly(ParameterInfo, FLinearColor::FromSRGBColor(ColorData[0]).R);
+				if (BakeOutput.PropertyIsLinearColor.Contains(Property) && BakeOutput.PropertyIsLinearColor.FindChecked(Property))
+				{
+					Material->SetScalarParameterValueEditorOnly(ParameterInfo, float(ColorData[0].R) / 255.0f);
+				}
+				else
+				{
+					Material->SetScalarParameterValueEditorOnly(ParameterInfo, FLinearColor::FromSRGBColor(ColorData[0]).R);
+				}
 			}
 		}
 	}
