Updated headers using UnrealCodeFixup to make sure dllstorage is on methods/staticvars instead of types. Part 4/n

diff --git a/Engine/Source/Developer/MaterialBaking/Private/MaterialBakingModule.h b/Engine/Source/Developer/MaterialBaking/Private/MaterialBakingModule.h
--- a/Engine/Source/Developer/MaterialBaking/Private/MaterialBakingModule.h
+++ b/Engine/Source/Developer/MaterialBaking/Private/MaterialBakingModule.h
@@ -10,6 +10,8 @@
 #include "IMaterialBakingModule.h"
 #include "MaterialPropertyEx.h"
 
+#define UE_API MATERIALBAKING_API
+
 class UTextureRenderTarget2D;
 class UMaterialOptions;
 class UMaterialInterface;
@@ -21,61 +23,61 @@ struct FMaterialData;
 struct FMeshData;
 struct FBakeOutput;
 
-class MATERIALBAKING_API FMaterialBakingModule : public IMaterialBakingModule
+class FMaterialBakingModule : public IMaterialBakingModule
 {
 public:
 	/** IModuleInterface overrides begin */
-	virtual void StartupModule();
-	virtual void ShutdownModule();
+	UE_API virtual void StartupModule();
+	UE_API virtual void ShutdownModule();
 	/** IModuleInterface overrides end */
 
 	/** Bakes out material properties according to MaterialSettings using MeshSettings and stores the output in Output */
-	virtual void BakeMaterials(const TArray<FMaterialData*>& MaterialSettings, const TArray<FMeshData*>& MeshSettings, TArray<FBakeOutput>& Output) override;
+	UE_API virtual void BakeMaterials(const TArray<FMaterialData*>& MaterialSettings, const TArray<FMeshData*>& MeshSettings, TArray<FBakeOutput>& Output) override;
 
 	/** Bakes out material properties according to extended MaterialSettings using MeshSettings and stores the output in Output */
-	virtual void BakeMaterials(const TArray<FMaterialDataEx*>& MaterialSettings, const TArray<FMeshData*>& MeshSettings, TArray<FBakeOutputEx>& Output) override;
+	UE_API virtual void BakeMaterials(const TArray<FMaterialDataEx*>& MaterialSettings, const TArray<FMeshData*>& MeshSettings, TArray<FBakeOutputEx>& Output) override;
 
 	/** Bakes out material properties according to MaterialSettings using MeshSettings and stores the output in Output */
-	virtual void BakeMaterials(const TArray<FMaterialData*>& MaterialSettings, const TArray<FMeshData*>& MeshSettings, FBakeOutput& Output) override;
+	UE_API virtual void BakeMaterials(const TArray<FMaterialData*>& MaterialSettings, const TArray<FMeshData*>& MeshSettings, FBakeOutput& Output) override;
 
 	/** Bakes out material properties according to extended MaterialSettings using MeshSettings and stores the output in Output */
-	virtual void BakeMaterials(const TArray<FMaterialDataEx*>& MaterialSettings, const TArray<FMeshData*>& MeshSettings, FBakeOutputEx& Output) override;
+	UE_API virtual void BakeMaterials(const TArray<FMaterialDataEx*>& MaterialSettings, const TArray<FMeshData*>& MeshSettings, FBakeOutputEx& Output) override;
 
 	/** Promps a slate window to allow the user to populate specific material baking settings used while baking out materials */
-	virtual bool SetupMaterialBakeSettings(TArray<TWeakObjectPtr<UObject>>& OptionObjects, int32 NumLODs) override;
+	UE_API virtual bool SetupMaterialBakeSettings(TArray<TWeakObjectPtr<UObject>>& OptionObjects, int32 NumLODs) override;
 
 	/** Outputs true HDR version of emissive color */
-	virtual void SetEmissiveHDR(bool bHDR) override;
+	UE_API virtual void SetEmissiveHDR(bool bHDR) override;
 
 	/** Bakes all material properties to linear textures, except for non-emissive colors */
-	virtual void SetLinearBake(bool bCorrectLinear) override;
+	UE_API virtual void SetLinearBake(bool bCorrectLinear) override;
 
 	/** Returns whether a specific material property is baked to a linear texture or not */
-	virtual bool IsLinearBake(FMaterialPropertyEx Property) override;
+	UE_API virtual bool IsLinearBake(FMaterialPropertyEx Property) override;
 
 	/** Obtain a CRC than can help trigger a rebake if code/global settings impacting the bake result change */
-	virtual uint32 GetCRC() const override;
+	UE_API virtual uint32 GetCRC() const override;
 
 protected:
 	friend class FMaterialBakingProcessor;
 
 	/* Creates and adds or reuses a RenderTarget from the pool */
-	UTextureRenderTarget2D* CreateRenderTarget(FMaterialPropertyEx InProperty, const FIntPoint& InTargetSize, bool bInUsePooledRenderTargets, const FColor& BackgroundColor);
+	UE_API UTextureRenderTarget2D* CreateRenderTarget(FMaterialPropertyEx InProperty, const FIntPoint& InTargetSize, bool bInUsePooledRenderTargets, const FColor& BackgroundColor);
 
 	/* Creates and adds (or reuses a ExportMaterialProxy from the pool if MaterialBaking.UseMaterialProxyCaching is set to 1) */
-	FExportMaterialProxy* CreateMaterialProxy(const FMaterialDataEx* MaterialSettings, const FMaterialPropertyEx& Property);
+	UE_API FExportMaterialProxy* CreateMaterialProxy(const FMaterialDataEx* MaterialSettings, const FMaterialPropertyEx& Property);
 
 	/** Cleans up all cached material proxies in MaterialProxyPool */
-	void CleanupMaterialProxies();
+	UE_API void CleanupMaterialProxies();
 
 	/** Free up all pooled render targets */
-	void CleanupRenderTargets();
+	UE_API void CleanupRenderTargets();
 
 	/** Callback for modified objects which should be removed from MaterialProxyPool in that case */
-	void OnObjectModified(UObject* Object);
+	UE_API void OnObjectModified(UObject* Object);
 
 	/** Callback used to clear material proxy cache on garbage collection */
-	void OnPreGarbageCollect();
+	UE_API void OnPreGarbageCollect();
 
 private:
 	enum EPropertyColorSpace
@@ -105,3 +107,5 @@ private:
 
 	bool bIsBakingMaterials;
 };
+
+#undef UE_API
diff --git a/Engine/Source/Developer/MaterialBaking/Public/IMaterialBakingModule.h b/Engine/Source/Developer/MaterialBaking/Public/IMaterialBakingModule.h
--- a/Engine/Source/Developer/MaterialBaking/Public/IMaterialBakingModule.h
+++ b/Engine/Source/Developer/MaterialBaking/Public/IMaterialBakingModule.h
@@ -21,7 +21,7 @@ struct FMeshData;
 struct FBakeOutput;
 struct FBakeOutputEx;
 
-class MATERIALBAKING_API IMaterialBakingModule : public IModuleInterface
+class IMaterialBakingModule : public IModuleInterface
 {
 public:
 	/** Bakes out material properties according to MaterialSettings using MeshSettings and stores the output in Output */
diff --git a/Engine/Source/Developer/MaterialBaking/Public/MaterialBakingHelpers.h b/Engine/Source/Developer/MaterialBaking/Public/MaterialBakingHelpers.h
--- a/Engine/Source/Developer/MaterialBaking/Public/MaterialBakingHelpers.h
+++ b/Engine/Source/Developer/MaterialBaking/Public/MaterialBakingHelpers.h
@@ -5,9 +5,11 @@
 #include "Containers/Array.h"
 #include "Math/Color.h"
 
+#define UE_API MATERIALBAKING_API
+
 struct FColor;
 
-class MATERIALBAKING_API FMaterialBakingHelpers
+class FMaterialBakingHelpers
 {
 public:
 	/**
@@ -18,7 +20,7 @@ public:
 	 * @param  MaxIterations   The max distance in pixels to smear the edges of the texture. When set to -1 this value is the max of the width and height of the buffer
 	 * @param  BackgroundColor The pixel color considered background in the image, and that can be smeared into
 	 */
-	static void PerformUVBorderSmear(TArray<FColor>& InOutPixels, int32 ImageWidth, int32 ImageHeight, int32 MaxIterations = -1, FColor BackgroundColor = FColor::Magenta);
+	static UE_API void PerformUVBorderSmear(TArray<FColor>& InOutPixels, int32 ImageWidth, int32 ImageHeight, int32 MaxIterations = -1, FColor BackgroundColor = FColor::Magenta);
 
 	/**
 	 * Shrinks the array to a single pixel if the entire output is monochrome. Params will be adjusted accordingly when this occurs.
@@ -28,7 +30,7 @@ public:
 	 * @param  ImageWidth      The width of the image in pixels
 	 * @param  ImageHeight     The height of the image in pixels
 	 */
-	static void PerformShrinking(TArray<FColor>& InOutPixels, int32& OnOutImageWidth, int32& InOutImageHeight, FColor BackgroundColor = FColor::Magenta);
+	static UE_API void PerformShrinking(TArray<FColor>& InOutPixels, int32& OnOutImageWidth, int32& InOutImageHeight, FColor BackgroundColor = FColor::Magenta);
 
 	/**
 	 * Applies a box blur to magenta pixels found in given texture represented by InOutPixels using non-background pixels, this creates a smear across the magenta/filled pixels.
@@ -39,5 +41,7 @@ public:
 	 *
 	 * @note   This version is allowed to shrink the array to a single pixel if the entire output is monochrome. Params will be adjusted accordingly when this occurs.
 	 */
-	static void PerformUVBorderSmearAndShrink(TArray<FColor>& InOutPixels, int32& InOutImageWidth, int32& InOutImageHeight, FColor BackgroundColor = FColor::Magenta);
+	static UE_API void PerformUVBorderSmearAndShrink(TArray<FColor>& InOutPixels, int32& InOutImageWidth, int32& InOutImageHeight, FColor BackgroundColor = FColor::Magenta);
 };
+
+#undef UE_API
diff --git a/Engine/Source/Developer/MaterialBaking/Public/MaterialOptions.h b/Engine/Source/Developer/MaterialBaking/Public/MaterialOptions.h
--- a/Engine/Source/Developer/MaterialBaking/Public/MaterialOptions.h
+++ b/Engine/Source/Developer/MaterialBaking/Public/MaterialOptions.h
@@ -60,8 +60,8 @@ struct FPropertyEntry
 };
 
 /** Options object to define what and how a material should be baked out */
-UCLASS(config = Editor, Blueprintable)
-class MATERIALBAKING_API UMaterialOptions : public UObject 
+UCLASS(MinimalAPI, config = Editor, Blueprintable)
+class UMaterialOptions : public UObject 
 {
 	GENERATED_BODY()
 
@@ -100,8 +100,8 @@ public:
 };
 
 /** Asset bake options object */
-UCLASS(Config = Editor, Blueprintable)
-class MATERIALBAKING_API UAssetBakeOptions : public UObject
+UCLASS(MinimalAPI, Config = Editor, Blueprintable)
+class UAssetBakeOptions : public UObject
 {
 	GENERATED_BODY()
 public:
@@ -112,8 +112,8 @@ public:
 };
 
 /** Material merge options object */
-UCLASS(Config = Editor, Blueprintable)
-class MATERIALBAKING_API UMaterialMergeOptions: public UObject
+UCLASS(MinimalAPI, Config = Editor, Blueprintable)
+class UMaterialMergeOptions: public UObject
 {
 	GENERATED_BODY()
 public:
@@ -129,4 +129,4 @@ public:
 	/** Blend mode for the final proxy material(s) */
 	UPROPERTY(EditAnywhere, BlueprintReadWrite, config, Category = OutputMaterialSettings, meta = (ExposeOnSpawn))
 	TEnumAsByte<EBlendMode> BlendMode;
-};
\ No newline at end of file
+};
diff --git a/Engine/Source/Developer/MaterialBaking/Public/MaterialPropertyEx.h b/Engine/Source/Developer/MaterialBaking/Public/MaterialPropertyEx.h
--- a/Engine/Source/Developer/MaterialBaking/Public/MaterialPropertyEx.h
+++ b/Engine/Source/Developer/MaterialBaking/Public/MaterialPropertyEx.h
@@ -9,8 +9,10 @@
 #include "UObject/NameTypes.h"
 #include "UObject/UnrealNames.h"
 
+#define UE_API MATERIALBAKING_API
+
 /** Structure extending EMaterialProperty to allow detailed information about custom output */
-struct MATERIALBAKING_API FMaterialPropertyEx
+struct FMaterialPropertyEx
 {
 	FMaterialPropertyEx(EMaterialProperty Type = MP_MAX, const FName& CustomOutput = NAME_None)
 		: Type(Type)
@@ -47,7 +49,7 @@ struct MATERIALBAKING_API FMaterialPropertyEx
 		return !Other.IsCustomOutput() ? GetTypeHash(Other.Type) : GetTypeHash(Other.CustomOutput);
 	}
 
-	FString ToString() const;
+	UE_API FString ToString() const;
 
 	/** The material property */
 	EMaterialProperty Type;
@@ -56,8 +58,10 @@ struct MATERIALBAKING_API FMaterialPropertyEx
 	FName CustomOutput;
 
 	/** The identifier for custom output Clear Coat Bottom Normal (used by shading model Clear Coat) */
-	static const FMaterialPropertyEx ClearCoatBottomNormal;
+	static UE_API const FMaterialPropertyEx ClearCoatBottomNormal;
 
 	/** The identifier for custom output Transmittance Color (used by shading model Thin Translucent) */
-	static const FMaterialPropertyEx TransmittanceColor;
+	static UE_API const FMaterialPropertyEx TransmittanceColor;
 };
+
+#undef UE_API
diff --git a/Engine/Source/Developer/MaterialUtilities/Public/MaterialUtilities.h b/Engine/Source/Developer/MaterialUtilities/Public/MaterialUtilities.h
--- a/Engine/Source/Developer/MaterialUtilities/Public/MaterialUtilities.h
+++ b/Engine/Source/Developer/MaterialUtilities/Public/MaterialUtilities.h
@@ -17,6 +17,8 @@
 
 #include "MaterialUtilities.generated.h"
 
+#define UE_API MATERIALUTILITIES_API
+
 class ALandscapeProxy;
 class Error;
 class ULandscapeComponent;
@@ -198,17 +200,17 @@ private:
 };
 
 /** Export material proxy cache*/
-struct MATERIALUTILITIES_API FExportMaterialProxyCache
+struct FExportMaterialProxyCache
 {
 	// Material proxies for each property. Note: we're not handling all properties here,
 	// so hold only up to MP_Normal inclusive.
 	class FMaterialRenderProxy* Proxies[MP_MAX];
 
-	FExportMaterialProxyCache();
-	~FExportMaterialProxyCache();
+	UE_API FExportMaterialProxyCache();
+	UE_API ~FExportMaterialProxyCache();
 
 	/** Releasing the cached render proxies */
-	void Release();
+	UE_API void Release();
 };
 
 /** Intermediate material merging data */
@@ -286,12 +288,12 @@ namespace UE::Geometry { class FDynamicMesh3; }
 /**
  * Material utilities
  */
-class MATERIALUTILITIES_API FMaterialUtilities : public IModuleInterface
+class FMaterialUtilities : public IModuleInterface
 {
 public:
 	/** Begin IModuleInterface implementation */
-	virtual void StartupModule() override;
-	virtual void ShutdownModule() override;
+	UE_API virtual void StartupModule() override;
+	UE_API virtual void ShutdownModule() override;
 	/** End IModuleInterface implementation */
 
 	/**
@@ -302,14 +304,14 @@ public:
 	*
 	* @return Size (width and height)
 	*/
-	static FIntPoint FindMaxTextureSize(UMaterialInterface* InMaterialInterface, FIntPoint MinimumSize = FIntPoint(1, 1));
+	static UE_API FIntPoint FindMaxTextureSize(UMaterialInterface* InMaterialInterface, FIntPoint MinimumSize = FIntPoint(1, 1));
 
 	/**
 	 * Whether material utilities support exporting specified material blend mode and property 
 	 */
 	UE_DEPRECATED(5.2, "Use SupportsExport(bool bIsOpaque, ...) function instead")
-	static bool SupportsExport(EBlendMode InBlendMode, EMaterialProperty InMaterialProperty);
-	static bool SupportsExport(bool bIsOpaque, EMaterialProperty InMaterialProperty);
+	static UE_API bool SupportsExport(EBlendMode InBlendMode, EMaterialProperty InMaterialProperty);
+	static UE_API bool SupportsExport(bool bIsOpaque, EMaterialProperty InMaterialProperty);
 
 	/**
 	 * Flattens specified landscape material
@@ -318,7 +320,7 @@ public:
 	 * @param OutFlattenMaterial	Output flattened material
 	 * @return						Whether operation was successful
 	 */
-	static bool ExportLandscapeMaterial(const ALandscapeProxy* InLandscape, FFlattenMaterial& OutFlattenMaterial);
+	static UE_API bool ExportLandscapeMaterial(const ALandscapeProxy* InLandscape, FFlattenMaterial& OutFlattenMaterial);
 
 	/**
 	 * Flattens specified landscape material
@@ -328,7 +330,7 @@ public:
 	 * @param OutFlattenMaterial	Output flattened material
 	 * @return						Whether operation was successful
 	 */
-	static bool ExportLandscapeMaterial(const ALandscapeProxy* InLandscape, const TSet<FPrimitiveComponentId>& HiddenPrimitives, FFlattenMaterial& OutFlattenMaterial);
+	static UE_API bool ExportLandscapeMaterial(const ALandscapeProxy* InLandscape, const TSet<FPrimitiveComponentId>& HiddenPrimitives, FFlattenMaterial& OutFlattenMaterial);
 	
 	/**
  	 * Generates a texture from an array of samples 
@@ -344,7 +346,7 @@ public:
 	 * @param SourceGuidHash		(optional) Hash (stored as Guid) to use part of the texture source's DDC key.
 	 * @return						The new texture.
 	 */
-	static UTexture2D* CreateTexture(UPackage* Outer, const FString& AssetLongName, FIntPoint Size, const TArray<FColor>& Samples, TextureCompressionSettings CompressionSettings, TextureGroup LODGroup, EObjectFlags Flags, bool bSRGB, const FGuid& SourceGuidHash = FGuid());
+	static UE_API UTexture2D* CreateTexture(UPackage* Outer, const FString& AssetLongName, FIntPoint Size, const TArray<FColor>& Samples, TextureCompressionSettings CompressionSettings, TextureGroup LODGroup, EObjectFlags Flags, bool bSRGB, const FGuid& SourceGuidHash = FGuid());
 
 	/**
 	 * Generates a texture from an array of samples
@@ -357,7 +359,7 @@ public:
 	 * @param Flags					ObjectFlags for the new texture
 	 * @return						The new texture.
 	 */
-	static UTexture2D* CreateTexture(UPackage* Outer, const FString& AssetLongName, FIntPoint Size, const TArray<FColor>& Samples, const FCreateTexture2DParameters& CreateParams, EObjectFlags Flags);
+	static UE_API UTexture2D* CreateTexture(UPackage* Outer, const FString& AssetLongName, FIntPoint Size, const TArray<FColor>& Samples, const FCreateTexture2DParameters& CreateParams, EObjectFlags Flags);
 
 	/**
 	 * Creates UMaterial object from a flatten material
@@ -368,7 +370,7 @@ public:
 	 * @param OutGeneratedAssets	List of generated assets - material, textures
 	 * @return						Returns a pointer to the constructed UMaterial object.
 	 */
-	static UMaterial* CreateMaterial(const FFlattenMaterial& InFlattenMaterial, UPackage* InOuter, const FString& BaseName, EObjectFlags Flags, const struct FMaterialProxySettings& MaterialProxySettings, TArray<UObject*>& OutGeneratedAssets, const TextureGroup& InTextureGroup = TEXTUREGROUP_World);
+	static UE_API UMaterial* CreateMaterial(const FFlattenMaterial& InFlattenMaterial, UPackage* InOuter, const FString& BaseName, EObjectFlags Flags, const struct FMaterialProxySettings& MaterialProxySettings, TArray<UObject*>& OutGeneratedAssets, const TextureGroup& InTextureGroup = TEXTUREGROUP_World);
 
 	/** 
 	* Creates an instanced material based of BaseMaterial
@@ -378,7 +380,7 @@ public:
 	* @return						Returns a pointer to the constructed UMaterialInstanceConstant object.
 	*/
 
-	static UMaterialInstanceConstant* CreateInstancedMaterial(UMaterialInterface* BaseMaterial, UPackage* InOuter, const FString& BaseName, EObjectFlags Flags);
+	static UE_API UMaterialInstanceConstant* CreateInstancedMaterial(UMaterialInterface* BaseMaterial, UPackage* InOuter, const FString& BaseName, EObjectFlags Flags);
 
 	/**
 	* Creates bakes textures for a ULandscapeComponent
@@ -386,7 +388,7 @@ public:
 	* @param LandscapeComponent		The component to bake textures for
 	* @return						Whether operation was successful
 	*/
-	static bool ExportBaseColor(ULandscapeComponent* LandscapeComponent, int32 TextureSize, TArray<FColor>& OutSamples);
+	static UE_API bool ExportBaseColor(ULandscapeComponent* LandscapeComponent, int32 TextureSize, TArray<FColor>& OutSamples);
 
 	/**
 	* Creates a FFlattenMaterial instance with the given MaterialProxySettings data
@@ -394,7 +396,7 @@ public:
 	* @param InMaterialLODSettings  Settings containing how to material should be merged
 	* @return						Set up FFlattenMaterial according to the given settings
 	*/
-	static FFlattenMaterial CreateFlattenMaterialWithSettings(const FMaterialProxySettings& InMaterialLODSettings);
+	static UE_API FFlattenMaterial CreateFlattenMaterialWithSettings(const FMaterialProxySettings& InMaterialLODSettings);
 
 	/**
 	* Analyzes given material to determine how many texture coordinates and whether or not vertex colors are used within the material Graph
@@ -404,9 +406,9 @@ public:
 	* @param OutNumTexCoords		Number of texture coordinates used across all properties flagged for merging
 	* @param OutRequiresVertexData	Flag whether or not Vertex Data is used in the material graph for the properties flagged for merging
 	*/
-	static void AnalyzeMaterial(class UMaterialInterface* InMaterial, const struct FMaterialProxySettings& InMaterialSettings, int32& OutNumTexCoords, bool& OutRequiresVertexData);
+	static UE_API void AnalyzeMaterial(class UMaterialInterface* InMaterial, const struct FMaterialProxySettings& InMaterialSettings, int32& OutNumTexCoords, bool& OutRequiresVertexData);
 
-	static void AnalyzeMaterial(class UMaterialInterface* InMaterial, const TArray<EMaterialProperty>& Properties, int32& OutNumTexCoords, bool& OutRequiresVertexData);
+	static UE_API void AnalyzeMaterial(class UMaterialInterface* InMaterial, const TArray<EMaterialProperty>& Properties, int32& OutNumTexCoords, bool& OutRequiresVertexData);
 
 	/**
 	* Remaps material indices where possible to reduce the number of materials required for creating a proxy material
@@ -421,30 +423,30 @@ public:
 	* @param OutMaterialMap					Map of MeshIDAndLOD keys with a material index array as value mapping InMeshData to the OutMaterials array
 	* @param OutMaterials					List of Material interfaces (unique)
 	*/
-	static void RemapUniqueMaterialIndices(const TArray<struct FSectionInfo>& InSections, const TArray<struct FRawMeshExt>& InMeshData, const TMap<FIntPoint, TArray<int32> >& InMaterialMap, const FMaterialProxySettings& InMaterialProxySettings, const bool bBakeVertexData, const bool bMergeMaterials, TArray<bool>& OutMeshShouldBakeVertexData, TMap<FIntPoint, TArray<int32> >& OutMaterialMap, TArray<struct FSectionInfo>& OutSections);
+	static UE_API void RemapUniqueMaterialIndices(const TArray<struct FSectionInfo>& InSections, const TArray<struct FRawMeshExt>& InMeshData, const TMap<FIntPoint, TArray<int32> >& InMaterialMap, const FMaterialProxySettings& InMaterialProxySettings, const bool bBakeVertexData, const bool bMergeMaterials, TArray<bool>& OutMeshShouldBakeVertexData, TMap<FIntPoint, TArray<int32> >& OutMaterialMap, TArray<struct FSectionInfo>& OutSections);
 
 	/**
 	* Tries to optimize the flatten material's data by picking out constant values for the various properties
 	*
 	* @param InFlattenMaterial				Flatten material to optimize
 	*/
-	static void OptimizeFlattenMaterial(FFlattenMaterial& InFlattenMaterial);
+	static UE_API void OptimizeFlattenMaterial(FFlattenMaterial& InFlattenMaterial);
 
 	/**
 	* Resizes flatten material's data if applicable by comparing it with the original settings
 	*
 	* @param InFlattenMaterial				Flatten material to optimize
 	*/
-	static void ResizeFlattenMaterial(FFlattenMaterial& InFlattenMaterial, const struct FMeshProxySettings& InMeshProxySettings);
+	static UE_API void ResizeFlattenMaterial(FFlattenMaterial& InFlattenMaterial, const struct FMeshProxySettings& InMeshProxySettings);
 
 	/** Tries to optimize the sample array (will set to const value if all samples are equal) */
-	static void OptimizeSampleArray(TArray<FColor>& InSamples, FIntPoint& InSampleSize);
+	static UE_API void OptimizeSampleArray(TArray<FColor>& InSamples, FIntPoint& InSampleSize);
 
 	/**
 	* Contains errors generated when exporting material texcoord scales. 
 	* Used to prevent displaying duplicates, as instances using the same shaders get the same issues.
 	*/
-	class MATERIALUTILITIES_API FExportErrorManager
+	class FExportErrorManager
 	{
 	public:
 
@@ -464,12 +466,12 @@ public:
 		* @param RegisterIndex		The register index bound to this texture.
 		* @param ErrorType			The issue encountered.
 		*/
-		void Register(const UMaterialInterface* Material, FName TextureName, int32 RegisterIndex, EErrorType ErrorType);
+		UE_API void Register(const UMaterialInterface* Material, FName TextureName, int32 RegisterIndex, EErrorType ErrorType);
 
 		/**
 		* Output all errors registered.
 		*/
-		void OutputToLog();
+		UE_API void OutputToLog();
 
 	private:
 
@@ -503,21 +505,21 @@ public:
 	* @param OutErrors			Manager to log errors (removes duplicates and similar errors)	
 	* @return					Whether operation was successful
 	*/
-	static bool ExportMaterialUVDensities(UMaterialInterface* InMaterial, EMaterialQualityLevel::Type QualityLevel, ERHIFeatureLevel::Type FeatureLevel, FExportErrorManager& OutErrors);
+	static UE_API bool ExportMaterialUVDensities(UMaterialInterface* InMaterial, EMaterialQualityLevel::Type QualityLevel, ERHIFeatureLevel::Type FeatureLevel, FExportErrorManager& OutErrors);
 
 	/** Calculates an importance value for the given set of materials according to the texture sampler usage */
-	static void DetermineMaterialImportance(const TArray<UMaterialInterface*>& InMaterials, TArray<float>& OutImportance);
+	static UE_API void DetermineMaterialImportance(const TArray<UMaterialInterface*>& InMaterials, TArray<float>& OutImportance);
 	/** Generates a set of texture boxes according to the given weights and final atlas texture size */
-	static void GeneratedBinnedTextureSquares(const FVector2D DestinationSize, TArray<float>& InTexureWeights, TArray<FBox2D>& OutGeneratedBoxes);
+	static UE_API void GeneratedBinnedTextureSquares(const FVector2D DestinationSize, TArray<float>& InTexureWeights, TArray<FBox2D>& OutGeneratedBoxes);
 
 	/** Creates a proxy material and the required texture assets */
-	static UMaterialInterface* CreateProxyMaterialAndTextures(UPackage* OuterPackage, const FString& AssetName, const FBakeOutput& BakeOutput, const FMeshData& MeshData, const FMaterialData& MaterialData, UMaterialOptions* Options);
+	static UE_API UMaterialInterface* CreateProxyMaterialAndTextures(UPackage* OuterPackage, const FString& AssetName, const FBakeOutput& BakeOutput, const FMeshData& MeshData, const FMaterialData& MaterialData, UMaterialOptions* Options);
 
 	/** Creates a proxy material and the required texture assets */
-	static UMaterialInterface* CreateProxyMaterialAndTextures(const FString& PackageName, const FString& AssetName, const FBakeOutput& BakeOutput, const FMeshData& MeshData, const FMaterialData& MaterialData, UMaterialOptions* Options);
+	static UE_API UMaterialInterface* CreateProxyMaterialAndTextures(const FString& PackageName, const FString& AssetName, const FBakeOutput& BakeOutput, const FMeshData& MeshData, const FMaterialData& MaterialData, UMaterialOptions* Options);
 
 	/** Creates a flatten material instance and the required texture assets */
-	static UMaterialInstanceConstant* CreateFlattenMaterialInstance(UPackage* InOuter, const FMaterialProxySettings& InMaterialProxySettings, UMaterialInterface* InBaseMaterial, const FFlattenMaterial& FlattenMaterial, const FString& AssetBasePath, const FString& AssetBaseName, TArray<UObject*>& OutAssetsToSync, FMaterialUpdateContext* MaterialUpdateContext = nullptr);
+	static UE_API UMaterialInstanceConstant* CreateFlattenMaterialInstance(UPackage* InOuter, const FMaterialProxySettings& InMaterialProxySettings, UMaterialInterface* InBaseMaterial, const FFlattenMaterial& FlattenMaterial, const FString& AssetBasePath, const FString& AssetBaseName, TArray<UObject*>& OutAssetsToSync, FMaterialUpdateContext* MaterialUpdateContext = nullptr);
 
 	/** Compute the required texel density needed to properly represent an object/objects covering the provided world space radius
 	* 
@@ -525,7 +527,7 @@ public:
 	* @param InWorldSpaceRadius	World space radius of the objects
 	* @return					The required texel density per METER
 	*/
-	static float ComputeRequiredTexelDensityFromScreenSize(const float InScreenSize, float InWorldSpaceRadius);
+	static UE_API float ComputeRequiredTexelDensityFromScreenSize(const float InScreenSize, float InWorldSpaceRadius);
 
 	/** Compute the required texel density needed to properly represent an object/objects covering the provided world space radius
 	*
@@ -533,34 +535,34 @@ public:
 	* @param InWorldSpaceRadius	World space radius of the objects
 	* @return					The required texel density per METER
 	*/
-	static float ComputeRequiredTexelDensityFromDrawDistance(const float InDrawDistance, float InWorldSpaceRadius);
+	static UE_API float ComputeRequiredTexelDensityFromDrawDistance(const float InDrawDistance, float InWorldSpaceRadius);
 
 	/** Compute the required texture size to achieve a target texel density for the given mesh.
 	* @param InMesh					The mesh for which we want to create a flatten material.
 	* @param InTargetTexelDensity	The target texel density.
 	* @return The texture size needed to achieve the required texel density.
 	*/
-	static int32 GetTextureSizeFromTargetTexelDensity(const FMeshDescription& InMesh, float InTargetTexelDensity);
+	static UE_API int32 GetTextureSizeFromTargetTexelDensity(const FMeshDescription& InMesh, float InTargetTexelDensity);
 
 	/** Compute the required texture size to achieve a target texel density for the given mesh.
 	* @param InMesh					The mesh for which we want to create a flatten material.
 	* @param InTargetTexelDensity	The target texel density.
 	* @return The texture size needed to achieve the required texel density.
 	*/
-	static int32 GetTextureSizeFromTargetTexelDensity(const UE::Geometry::FDynamicMesh3& Mesh, float TargetTexelDensity);
+	static UE_API int32 GetTextureSizeFromTargetTexelDensity(const UE::Geometry::FDynamicMesh3& Mesh, float TargetTexelDensity);
 
 	/** Compute the required texture size to achieve a target texel density, given the world to UV space ratio.
 	* @param InMesh3DArea	World space area of the mesh
 	* @param InMeshUVArea	UV space area of the mesh
 	* @return The texture size needed to achieve the required texel density.
 	*/
-	static int32 GetTextureSizeFromTargetTexelDensity(double InMesh3DArea, double InMeshUVArea, double InTargetTexelDensity);
+	static UE_API int32 GetTextureSizeFromTargetTexelDensity(double InMesh3DArea, double InMeshUVArea, double InTargetTexelDensity);
 
 	/** Validate that the provided material has all the required parameters needed to be considered a flattening material */
-	static bool IsValidFlattenMaterial(const UMaterialInterface* InBaseMaterial);
+	static UE_API bool IsValidFlattenMaterial(const UMaterialInterface* InBaseMaterial);
 
 	/** Get the name of the Texture parameter associated with a given flatten material property */
-	static FString GetFlattenMaterialTextureName(EFlattenMaterialProperties InProperty, UMaterialInterface* InBaseMaterial);
+	static UE_API FString GetFlattenMaterialTextureName(EFlattenMaterialProperties InProperty, UMaterialInterface* InBaseMaterial);
 
 private:
 	
@@ -572,16 +574,18 @@ private:
 	* @param InTargetSize			Dimensions of the render target
 	* @return						Created render target
 	*/
-	static UTextureRenderTarget2D* CreateRenderTarget(bool bInForceLinearGamma, bool bNormalMap, EPixelFormat InPixelFormat, FIntPoint& InTargetSize );
+	static UE_API UTextureRenderTarget2D* CreateRenderTarget(bool bInForceLinearGamma, bool bNormalMap, EPixelFormat InPixelFormat, FIntPoint& InTargetSize );
 	
 	/** Clears the pool with available render targets */
-	static void ClearRenderTargetPool();	
+	static UE_API void ClearRenderTargetPool();	
 
 	/** Call back for garbage collector, cleans up the RenderTargetPool if CurrentlyRendering is set to false */
-	void OnPreGarbageCollect();
+	UE_API void OnPreGarbageCollect();
 private:
 	/** Flag to indicate whether or not a texture is currently being rendered out */
-	static bool CurrentlyRendering;
+	static UE_API bool CurrentlyRendering;
 	/** Pool of available render targets, cached for re-using on consecutive property rendering */
-	static TArray<UTextureRenderTarget2D*> RenderTargetPool;
+	static UE_API TArray<UTextureRenderTarget2D*> RenderTargetPool;
 };
+
+#undef UE_API
diff --git a/Engine/Source/Developer/Merge/Private/SBlueprintMerge.h b/Engine/Source/Developer/Merge/Private/SBlueprintMerge.h
--- a/Engine/Source/Developer/Merge/Private/SBlueprintMerge.h
+++ b/Engine/Source/Developer/Merge/Private/SBlueprintMerge.h
@@ -8,15 +8,17 @@
 #include "MergeUtils.h"
 #include "Widgets/Views/STreeView.h"
 
+#define UE_API MERGE_API
+
 class FBlueprintDifferenceTreeEntry;
 class UBlueprint;
 struct FAssetRevisionInfo;
 struct FDiffSingleResult;
 
-class MERGE_API SBlueprintMerge : public SCompoundWidget
+class SBlueprintMerge : public SCompoundWidget
 {
 public:
-	SBlueprintMerge();
+	UE_API SBlueprintMerge();
 
 	SLATE_BEGIN_ARGS(SBlueprintMerge)
 		: _bForcePickAssets(false)
@@ -25,30 +27,30 @@ public:
 		SLATE_EVENT(FOnMergeResolved, OnMergeResolved)
 	SLATE_END_ARGS()
 
-	void Construct(const FArguments InArgs, const FBlueprintMergeData& InData);
+	UE_API void Construct(const FArguments InArgs, const FBlueprintMergeData& InData);
 
 private:
 	/** Helper functions */
-	UBlueprint* GetTargetBlueprint();
+	UE_API UBlueprint* GetTargetBlueprint();
 
 	/** Event handlers */
-	void NextDiff();
-	void PrevDiff();
-	bool HasNextDiff() const;
-	bool HasPrevDiff() const;
-
-	void NextConflict();
-	void PrevConflict();
-	bool HasNextConflict() const;
-	bool HasPrevConflict() const;
-
-	void OnStartMerge();
-	void OnFinishMerge();
-	void OnCancelClicked();
-	void OnModeChanged(FName NewMode);
-	void OnAcceptRemote();
-	void OnAcceptLocal();
-	void ResolveMerge(UBlueprint* Result);
+	UE_API void NextDiff();
+	UE_API void PrevDiff();
+	UE_API bool HasNextDiff() const;
+	UE_API bool HasPrevDiff() const;
+
+	UE_API void NextConflict();
+	UE_API void PrevConflict();
+	UE_API bool HasNextConflict() const;
+	UE_API bool HasPrevConflict() const;
+
+	UE_API void OnStartMerge();
+	UE_API void OnFinishMerge();
+	UE_API void OnCancelClicked();
+	UE_API void OnModeChanged(FName NewMode);
+	UE_API void OnAcceptRemote();
+	UE_API void OnAcceptLocal();
+	UE_API void ResolveMerge(UBlueprint* Result);
 
 	/** 
 	 * If the user has yet to pick their remote/base/local assets, then they're
@@ -57,7 +59,7 @@ private:
 	 *
 	 * @return False if the AssetPickerControl should be up, otherwise true.
 	 */
-	bool IsActivelyMerging() const;
+	UE_API bool IsActivelyMerging() const;
 
 	/** 
 	 * The user cannot start a merge until they have chosen a remote, base, &
@@ -66,7 +68,7 @@ private:
 	 *
 	 * @return True, if the user has selected all three merge assets.
 	 */
-	bool CanStartMerge() const;
+	UE_API bool CanStartMerge() const;
 
 	/** 
 	 * Callback function, utilized every time the user picks an asset/revision  
@@ -76,7 +78,7 @@ private:
 	 * @param  AssetId		Defines the asset that was altered (remote, base, or local).
 	 * @param  AssetInfo	Defines the asset selection (name, revision, etc.)
 	 */
-	void OnMergeAssetSelected(EMergeAssetId::Type AssetId, const FAssetRevisionInfo& AssetInfo);
+	UE_API void OnMergeAssetSelected(EMergeAssetId::Type AssetId, const FAssetRevisionInfo& AssetInfo);
 
 	FBlueprintMergeData		Data; 
 	FString					BackupSubDir;
@@ -128,3 +130,5 @@ private:
 	TArray< TSharedPtr<struct FDiffSingleResult> > LocalDiffResults;
 	TArray< TSharedPtr<struct FDiffSingleResult> > RemoteDiffResults;
 };
+
+#undef UE_API
diff --git a/Engine/Source/Developer/MeshBuilder/Public/MeshBuilder.h b/Engine/Source/Developer/MeshBuilder/Public/MeshBuilder.h
--- a/Engine/Source/Developer/MeshBuilder/Public/MeshBuilder.h
+++ b/Engine/Source/Developer/MeshBuilder/Public/MeshBuilder.h
@@ -6,6 +6,8 @@
 #include "CoreMinimal.h"
 #include "Math/UnrealMathSSE.h"
 
+#define UE_API MESHBUILDER_API
+
 class FStaticMeshLODGroup;
 class FStaticMeshRenderData;
 class FStaticMeshSectionArray;
@@ -19,10 +21,10 @@ struct FMeshBuildVertexData;
  * Abstract class which is the base class of all builder.
  * All share code to build some render data should be found inside this class
  */
-class MESHBUILDER_API FMeshBuilder
+class FMeshBuilder
 {
 public:
-	FMeshBuilder();
+	UE_API FMeshBuilder();
 	
 	/**
 	 * Build function should be override and is the starting point for static mesh builders
@@ -30,7 +32,7 @@ public:
 	virtual bool Build(FStaticMeshRenderData& OutRenderData, const FStaticMeshBuildParameters& BuildParameters) = 0;
 
 	UE_DEPRECATED(5.5, "Use FStaticMeshBuildParameters instead.")
-	virtual bool Build(
+	UE_API virtual bool Build(
 		FStaticMeshRenderData& OutRenderData,
 		UStaticMesh* StaticMesh,
 		const FStaticMeshLODGroup& LODGroup,
@@ -48,8 +50,10 @@ public:
 	virtual bool Build(FSkeletalMeshRenderData& OutRenderData, const FSkeletalMeshBuildParameters& SkeletalMeshBuildParameters) = 0;
 
 	UE_DEPRECATED(5.6, "Use the overload that takes FSkeletalMeshRenderData instead")
-	virtual bool Build(const FSkeletalMeshBuildParameters& SkeletalMeshBuildParameters) final;
+	UE_API virtual bool Build(const FSkeletalMeshBuildParameters& SkeletalMeshBuildParameters) final;
 
 private:
 
 };
+
+#undef UE_API
diff --git a/Engine/Source/Developer/MeshBuilder/Public/MeshDescriptionHelper.h b/Engine/Source/Developer/MeshBuilder/Public/MeshDescriptionHelper.h
--- a/Engine/Source/Developer/MeshBuilder/Public/MeshDescriptionHelper.h
+++ b/Engine/Source/Developer/MeshBuilder/Public/MeshDescriptionHelper.h
@@ -6,6 +6,8 @@
 #include "Logging/LogMacros.h"
 #include "OverlappingCorners.h"
 
+#define UE_API MESHBUILDER_API
+
 DECLARE_LOG_CATEGORY_EXTERN(LogMeshBuilder, Log, All);
 
 class UObject;
@@ -14,18 +16,18 @@ struct FMeshBuildSettings;
 struct FVertexInstanceID;
 struct FMeshReductionSettings;
 
-class MESHBUILDER_API FMeshDescriptionHelper
+class FMeshDescriptionHelper
 {
 public:
 
-	FMeshDescriptionHelper(FMeshBuildSettings* InBuildSettings);
+	UE_API FMeshDescriptionHelper(FMeshBuildSettings* InBuildSettings);
 
 	//Build a render mesh description with the BuildSettings. This will update the RenderMeshDescription in place
-	void SetupRenderMeshDescription(UObject* Owner, FMeshDescription& RenderMeshDescription, bool bForNanite, bool bNeedTangents);
+	UE_API void SetupRenderMeshDescription(UObject* Owner, FMeshDescription& RenderMeshDescription, bool bForNanite, bool bNeedTangents);
 
-	void ReduceLOD(const FMeshDescription& BaseMesh, FMeshDescription& DestMesh, const struct FMeshReductionSettings& ReductionSettings, const FOverlappingCorners& InOverlappingCorners, float &OutMaxDeviation);
+	UE_API void ReduceLOD(const FMeshDescription& BaseMesh, FMeshDescription& DestMesh, const struct FMeshReductionSettings& ReductionSettings, const FOverlappingCorners& InOverlappingCorners, float &OutMaxDeviation);
 
-	void FindOverlappingCorners(const FMeshDescription& MeshDescription, float ComparisonThreshold);
+	UE_API void FindOverlappingCorners(const FMeshDescription& MeshDescription, float ComparisonThreshold);
 
 	const FOverlappingCorners& GetOverlappingCorners() const { return OverlappingCorners; }
 
@@ -57,3 +59,5 @@ private:
 	};
 
 };
+
+#undef UE_API
diff --git a/Engine/Source/Developer/MeshBuilder/Public/SkeletalMeshBuilder.h b/Engine/Source/Developer/MeshBuilder/Public/SkeletalMeshBuilder.h
--- a/Engine/Source/Developer/MeshBuilder/Public/SkeletalMeshBuilder.h
+++ b/Engine/Source/Developer/MeshBuilder/Public/SkeletalMeshBuilder.h
@@ -5,6 +5,8 @@
 #include "CoreMinimal.h"
 #include "MeshBuilder.h"
 
+#define UE_API MESHBUILDER_API
+
 DECLARE_LOG_CATEGORY_EXTERN(LogSkeletalMeshBuilder, Log, All);
 
 class UStaticMesh;
@@ -13,10 +15,10 @@ class FStaticMeshLODGroup;
 class USkeletalMesh;
 class FSkeletalMeshRenderData;
 
-class MESHBUILDER_API FSkeletalMeshBuilder : public FMeshBuilder
+class FSkeletalMeshBuilder : public FMeshBuilder
 {
 public:
-	FSkeletalMeshBuilder();
+	UE_API FSkeletalMeshBuilder();
 
 	//No support for static mesh build in this class
 	virtual bool Build(FStaticMeshRenderData& OutRenderData, const struct FStaticMeshBuildParameters& BuildParameters) override
@@ -37,7 +39,7 @@ public:
 		return false;
 	}
 	
-	virtual bool Build(FSkeletalMeshRenderData& OutRenderData, const struct FSkeletalMeshBuildParameters& SkeletalMeshBuildParameters) override;
+	UE_API virtual bool Build(FSkeletalMeshRenderData& OutRenderData, const struct FSkeletalMeshBuildParameters& SkeletalMeshBuildParameters) override;
 
 	virtual ~FSkeletalMeshBuilder() {}
 
@@ -45,3 +47,4 @@ private:
 
 };
 
+#undef UE_API
diff --git a/Engine/Source/Developer/MeshBuilder/Public/StaticMeshBuilder.h b/Engine/Source/Developer/MeshBuilder/Public/StaticMeshBuilder.h
--- a/Engine/Source/Developer/MeshBuilder/Public/StaticMeshBuilder.h
+++ b/Engine/Source/Developer/MeshBuilder/Public/StaticMeshBuilder.h
@@ -5,6 +5,8 @@
 #include "CoreMinimal.h"
 #include "MeshBuilder.h"
 
+#define UE_API MESHBUILDER_API
+
 DECLARE_LOG_CATEGORY_EXTERN(LogStaticMeshBuilder, Log, All);
 
 class ITargetPlatform;
@@ -15,13 +17,13 @@ struct FOverlappingCorners;
 struct FMeshDescription;
 struct FMeshBuildSettings;
 
-class MESHBUILDER_API FStaticMeshBuilder : public FMeshBuilder
+class FStaticMeshBuilder : public FMeshBuilder
 {
 public:
-	FStaticMeshBuilder();
+	UE_API FStaticMeshBuilder();
 	virtual ~FStaticMeshBuilder() {}
 
-	virtual bool Build(FStaticMeshRenderData& OutRenderData, const FStaticMeshBuildParameters& BuildParameters) override;
+	UE_API virtual bool Build(FStaticMeshRenderData& OutRenderData, const FStaticMeshBuildParameters& BuildParameters) override;
 
 	//No support for skeletal mesh build in this class
 	virtual bool Build(FSkeletalMeshRenderData& OutRenderData, const struct FSkeletalMeshBuildParameters& SkeletalMeshBuildParameters) override
@@ -31,7 +33,7 @@ public:
 		return false;
 	}
 
-	virtual bool BuildMeshVertexPositions(
+	UE_API virtual bool BuildMeshVertexPositions(
 		UStaticMesh* StaticMesh,
 		TArray<uint32>& Indices,
 		TArray<FVector3f>& Vertices,
@@ -39,8 +41,8 @@ public:
 
 private:
 
-	void OnBuildRenderMeshStart(class UStaticMesh* StaticMesh, const bool bInvalidateLighting);
-	void OnBuildRenderMeshFinish(class UStaticMesh* StaticMesh, const bool bRebuildBoundsAndCollision);
+	UE_API void OnBuildRenderMeshStart(class UStaticMesh* StaticMesh, const bool bInvalidateLighting);
+	UE_API void OnBuildRenderMeshFinish(class UStaticMesh* StaticMesh, const bool bRebuildBoundsAndCollision);
 
 	/** Used to refresh all components in the scene that may be using a mesh we're editing */
 	TSharedPtr<class FStaticMeshComponentRecreateRenderStateContext> RecreateRenderStateContext;
@@ -69,4 +71,6 @@ namespace UE::Private::StaticMeshBuilder
 		FStaticMeshSectionArray& SectionsOut,
 		TArray<uint32>& CombinedIndicesOut,
 		bool& bNeeds32BitIndicesOut);
-}
\ No newline at end of file
+}
+
+#undef UE_API
diff --git a/Engine/Source/Developer/MeshMergeUtilities/Private/MeshMergeHelpers.h b/Engine/Source/Developer/MeshMergeUtilities/Private/MeshMergeHelpers.h
--- a/Engine/Source/Developer/MeshMergeUtilities/Private/MeshMergeHelpers.h
+++ b/Engine/Source/Developer/MeshMergeUtilities/Private/MeshMergeHelpers.h
@@ -5,6 +5,8 @@
 #include "CoreMinimal.h"
 #include "MeshMerge/MeshProxySettings.h"
 
+#define UE_API MESHMERGEUTILITIES_API
+
 
 class USkeletalMeshComponent;
 class UStaticMeshComponent;
@@ -20,62 +22,64 @@ struct FStaticMeshLODResources;
 struct FKAggregateGeom;
 class UInstancedStaticMeshComponent;
 
-class MESHMERGEUTILITIES_API FMeshMergeHelpers
+class FMeshMergeHelpers
 {
 public:
 	/** Extracting section info data from static, skeletal mesh (components) */
-	static void ExtractSections(const UStaticMeshComponent* Component, int32 LODIndex, TArray<FSectionInfo>& OutSections);
-	static void ExtractSections(const USkeletalMeshComponent* Component, int32 LODIndex, TArray<FSectionInfo>& OutSections);
-	static void ExtractSections(const UStaticMesh* StaticMesh, int32 LODIndex, TArray<FSectionInfo>& OutSections);
+	static UE_API void ExtractSections(const UStaticMeshComponent* Component, int32 LODIndex, TArray<FSectionInfo>& OutSections);
+	static UE_API void ExtractSections(const USkeletalMeshComponent* Component, int32 LODIndex, TArray<FSectionInfo>& OutSections);
+	static UE_API void ExtractSections(const UStaticMesh* StaticMesh, int32 LODIndex, TArray<FSectionInfo>& OutSections);
 
 	/** Expanding instance data from instanced static mesh components */
-	static void ExpandInstances(const UInstancedStaticMeshComponent* InInstancedStaticMeshComponent, FMeshDescription& InOutRawMesh);
+	static UE_API void ExpandInstances(const UInstancedStaticMeshComponent* InInstancedStaticMeshComponent, FMeshDescription& InOutRawMesh);
 
 	/** Extracting mesh data in FMeshDescription form from static, skeletal mesh (components) */
-	static void RetrieveMesh(const UStaticMeshComponent* StaticMeshComponent, int32 LODIndex, FMeshDescription& OutMeshDescription, bool bPropagateVertexColours, bool bApplyComponentTransform = true);
-	static void RetrieveMesh(const USkeletalMeshComponent* SkeletalMeshComponent, int32 LODIndex, FMeshDescription& OutMeshDescription, bool bPropagateVertexColours, bool bApplyComponentTransform = false);
-	static void RetrieveMesh(const UStaticMesh* StaticMesh, int32 LODIndex, FMeshDescription& OutMeshDescription);
+	static UE_API void RetrieveMesh(const UStaticMeshComponent* StaticMeshComponent, int32 LODIndex, FMeshDescription& OutMeshDescription, bool bPropagateVertexColours, bool bApplyComponentTransform = true);
+	static UE_API void RetrieveMesh(const USkeletalMeshComponent* SkeletalMeshComponent, int32 LODIndex, FMeshDescription& OutMeshDescription, bool bPropagateVertexColours, bool bApplyComponentTransform = false);
+	static UE_API void RetrieveMesh(const UStaticMesh* StaticMesh, int32 LODIndex, FMeshDescription& OutMeshDescription);
 	
 	/** Exports static mesh LOD render data to a RawMesh */
-	static void ExportStaticMeshLOD(const FStaticMeshLODResources& StaticMeshLOD, FMeshDescription& OutMeshDescription, const TArray<FStaticMaterial>& Materials);
+	static UE_API void ExportStaticMeshLOD(const FStaticMeshLODResources& StaticMeshLOD, FMeshDescription& OutMeshDescription, const TArray<FStaticMaterial>& Materials);
 
 	/** Checks whether or not the texture coordinates are outside of 0-1 UV ranges */
-	static bool CheckWrappingUVs(const TArray<FVector2D>& UVs);
-	static bool CheckWrappingUVs(const FMeshDescription& MeshDescription, int32 UVChannelIndex);
+	static UE_API bool CheckWrappingUVs(const TArray<FVector2D>& UVs);
+	static UE_API bool CheckWrappingUVs(const FMeshDescription& MeshDescription, int32 UVChannelIndex);
 
 	/** Culls away triangles which are inside culling volumes or completely underneath the landscape */
-	static void CullTrianglesFromVolumesAndUnderLandscapes(const UWorld* World, const FBoxSphereBounds& Bounds, FMeshDescription& InOutRawMesh);
+	static UE_API void CullTrianglesFromVolumesAndUnderLandscapes(const UWorld* World, const FBoxSphereBounds& Bounds, FMeshDescription& InOutRawMesh);
 	
 	/** Propagates deformation along spline to mesh data */
-	static void PropagateSplineDeformationToMesh(const USplineMeshComponent* InSplineMeshComponent, FMeshDescription& InOutMeshDescription);
+	static UE_API void PropagateSplineDeformationToMesh(const USplineMeshComponent* InSplineMeshComponent, FMeshDescription& InOutMeshDescription);
 	
 	/** Propagates deformation along spline to physics geometry data */
-	static void PropagateSplineDeformationToPhysicsGeometry(USplineMeshComponent* SplineMeshComponent, FKAggregateGeom& InOutPhysicsGeometry);
+	static UE_API void PropagateSplineDeformationToPhysicsGeometry(USplineMeshComponent* SplineMeshComponent, FKAggregateGeom& InOutPhysicsGeometry);
 
 	/** Retrieves all culling landscapes and volumes as FMeshDescription structures. Note the caller is responsible for deleting the heap data managed by OutCullingMeshes */
-	static void RetrieveCullingLandscapeAndVolumes(UWorld* InWorld, const FBoxSphereBounds& EstimatedMeshProxyBounds, const TEnumAsByte<ELandscapeCullingPrecision::Type> PrecisionType, TArray<FMeshDescription*>& OutCullingMeshes);
+	static UE_API void RetrieveCullingLandscapeAndVolumes(UWorld* InWorld, const FBoxSphereBounds& EstimatedMeshProxyBounds, const TEnumAsByte<ELandscapeCullingPrecision::Type> PrecisionType, TArray<FMeshDescription*>& OutCullingMeshes);
 
 	/** Transforms physics geometry data using InTransform */
-	static void TransformPhysicsGeometry(const FTransform& InTransform, const bool bBakeConvexTransform, struct FKAggregateGeom& AggGeom);
+	static UE_API void TransformPhysicsGeometry(const FTransform& InTransform, const bool bBakeConvexTransform, struct FKAggregateGeom& AggGeom);
 	
 	/** Extract physics geometry data from a body setup */
-	static void ExtractPhysicsGeometry(UBodySetup* InBodySetup, const FTransform& ComponentToWorld, const bool bBakeConvexTransform, struct FKAggregateGeom& OutAggGeom);
+	static UE_API void ExtractPhysicsGeometry(UBodySetup* InBodySetup, const FTransform& ComponentToWorld, const bool bBakeConvexTransform, struct FKAggregateGeom& OutAggGeom);
 
 	/** Ensure that UV is in valid 0-1 UV ranges */
-	static FVector2D GetValidUV(const FVector2D& UV);
+	static UE_API FVector2D GetValidUV(const FVector2D& UV);
 	
 	/** Calculates UV coordinates bounds for the given MeshDescription	*/
-	static void CalculateTextureCoordinateBoundsForMesh(const FMeshDescription& InMeshDescription, TArray<FBox2D>& OutBounds);
+	static UE_API void CalculateTextureCoordinateBoundsForMesh(const FMeshDescription& InMeshDescription, TArray<FBox2D>& OutBounds);
 
 	/** Propagates vertex painted colors from the StaticMeshComponent instance to MeshDescription */
-	static bool PropagatePaintedColorsToMesh(const UStaticMeshComponent* StaticMeshComponent, int32 LODIndex, FMeshDescription& InOutMeshDescription);
+	static UE_API bool PropagatePaintedColorsToMesh(const UStaticMeshComponent* StaticMeshComponent, int32 LODIndex, FMeshDescription& InOutMeshDescription);
 
 	/** Checks whether or not the landscape proxy is hit given a ray start and end */
-	static bool IsLandscapeHit(const FVector& RayOrigin, const FVector& RayEndPoint, const UWorld* World, const TArray<ALandscapeProxy*>& LandscapeProxies, FVector& OutHitLocation);
+	static UE_API bool IsLandscapeHit(const FVector& RayOrigin, const FVector& RayEndPoint, const UWorld* World, const TArray<ALandscapeProxy*>& LandscapeProxies, FVector& OutHitLocation);
 	
 	/** Merges imposter meshes into the given MeshDescription. */
-	static void MergeImpostersToMesh(TArray<const UStaticMeshComponent*> ImposterComponents, FMeshDescription& InOutMeshDescription, const FVector& InPivot, int32 BaseMaterialIndex, TArray<UMaterialInterface*>& OutImposterMaterials);
+	static UE_API void MergeImpostersToMesh(TArray<const UStaticMeshComponent*> ImposterComponents, FMeshDescription& InOutMeshDescription, const FVector& InPivot, int32 BaseMaterialIndex, TArray<UMaterialInterface*>& OutImposterMaterials);
 
 	/** Ensure a generated HLOD mesh is not referencing non standalone materials. */
-	static void FixupNonStandaloneMaterialReferences(UStaticMesh* InStaticMesh);
+	static UE_API void FixupNonStandaloneMaterialReferences(UStaticMesh* InStaticMesh);
 };
+
+#undef UE_API
diff --git a/Engine/Source/Developer/MeshMergeUtilities/Public/IMeshMergeUtilities.h b/Engine/Source/Developer/MeshMergeUtilities/Public/IMeshMergeUtilities.h
--- a/Engine/Source/Developer/MeshMergeUtilities/Public/IMeshMergeUtilities.h
+++ b/Engine/Source/Developer/MeshMergeUtilities/Public/IMeshMergeUtilities.h
@@ -35,7 +35,7 @@ enum class EFlattenMaterialProperties : uint8;
 // Typedef pair of FVector2Ds, key == uv position offset and value == uv scaling factor
 typedef TPair<FVector2D, FVector2D> FUVOffsetScalePair;
 
-class MESHMERGEUTILITIES_API IMeshMergeUtilities
+class IMeshMergeUtilities
 {
 public:
 	virtual ~IMeshMergeUtilities() {};
@@ -143,4 +143,4 @@ public:
 
 	virtual void RegisterExtension(IMeshMergeExtension* InExtension) = 0;
 	virtual void UnregisterExtension(IMeshMergeExtension* InExtension) = 0;
-};
\ No newline at end of file
+};
diff --git a/Engine/Source/Developer/MeshMergeUtilities/Public/MeshMergeDataTracker.h b/Engine/Source/Developer/MeshMergeUtilities/Public/MeshMergeDataTracker.h
--- a/Engine/Source/Developer/MeshMergeUtilities/Public/MeshMergeDataTracker.h
+++ b/Engine/Source/Developer/MeshMergeUtilities/Public/MeshMergeDataTracker.h
@@ -7,6 +7,8 @@
 #include "MaterialUtilities.h"
 #include "StaticMeshAttributes.h"
 
+#define UE_API MESHMERGEUTILITIES_API
+
 typedef TPair<uint32, uint32> MeshLODPair;
 typedef TPair<uint32, uint32> SectionRemapPair;
 typedef TPair<uint32, uint32> MaterialRemapPair;
@@ -14,7 +16,7 @@ typedef TPair<uint32, uint32> MaterialRemapPair;
 struct FSectionInfo;
 
 /** Structure representing a mesh and lod index */
-struct MESHMERGEUTILITIES_API FMeshLODKey
+struct FMeshLODKey
 {
 public:
 	FMeshLODKey(int32 MeshIndex, int32 LODIndex, UStaticMesh* InMesh = nullptr)
@@ -73,7 +75,7 @@ protected:
 };
 
 /** This helps identify materials that can be collapsed because of equivalence */
-struct MESHMERGEUTILITIES_API FMaterialKey
+struct FMaterialKey
 {
 	FMaterialKey(const UMaterialInterface* InMaterial)
 		: Material(InMaterial)
@@ -105,87 +107,87 @@ typedef TMap<FMeshLODKey, FMeshDescription>::TIterator TRawMeshIterator;
 typedef TArray<int32>::TConstIterator TConstLODIndexIterator;
 
 /** Used to keep track of in-flight data while meshes are merged and their corresponding materials baked down */
-class MESHMERGEUTILITIES_API FMeshMergeDataTracker
+class FMeshMergeDataTracker
 {	
 public:
-	FMeshMergeDataTracker();
+	UE_API FMeshMergeDataTracker();
 
 	/** Looks at all available raw mesh data and processes it to populate some flags */
-	void ProcessRawMeshes();
+	UE_API void ProcessRawMeshes();
 
 	/** Adding a mapping between the index for an original mesh section index and the mesh section it will be indexed to in the final mesh */
-	void AddSectionRemapping(int32 MeshIndex, int32 LODIndex, int32 OriginalIndex, int32 UniqueIndex);
+	UE_API void AddSectionRemapping(int32 MeshIndex, int32 LODIndex, int32 OriginalIndex, int32 UniqueIndex);
 	/** Retrieves the MeshLOD keys from which the original sections are mapped to the unique section index */
-	void GetMeshLODsMappedToUniqueSection(int32 UniqueIndex, TArray<FMeshLODKey>& InOutMeshLODs);
+	UE_API void GetMeshLODsMappedToUniqueSection(int32 UniqueIndex, TArray<FMeshLODKey>& InOutMeshLODs);
 	/** Retrieves all section mappings for the MeshLOD key*/
-	void GetMappingsForMeshLOD(FMeshLODKey Key, TArray<SectionRemapPair>& InOutMappings);
+	UE_API void GetMappingsForMeshLOD(FMeshLODKey Key, TArray<SectionRemapPair>& InOutMappings);
 	
 	/** Adds or retrieves raw mesh data for the mesh and LOD index */
-	FMeshDescription& AddAndRetrieveRawMesh(int32 MeshIndex, int32 LODIndex, UStaticMesh* InMesh);
+	UE_API FMeshDescription& AddAndRetrieveRawMesh(int32 MeshIndex, int32 LODIndex, UStaticMesh* InMesh);
 	/** Removes raw mesh entry for the given mesh and LOD index */
-	void RemoveRawMesh(int32 MeshIndex, int32 LODIndex);
+	UE_API void RemoveRawMesh(int32 MeshIndex, int32 LODIndex);
 	/** Retrieves Raw Mesh ptr for the given mesh and LOD index */
-	FMeshDescription* GetRawMeshPtr(int32 MeshIndex, int32 LODIndex);
+	UE_API FMeshDescription* GetRawMeshPtr(int32 MeshIndex, int32 LODIndex);
 	/** Retrieves Raw Mesh ptr for the given MeshLOD key */
-	FMeshDescription* GetRawMeshPtr(FMeshLODKey Key);
+	UE_API FMeshDescription* GetRawMeshPtr(FMeshLODKey Key);
 	/** Tries to retrieve a FMeshDescription and returns the LOD index it found an entry for */
-	FMeshDescription* FindRawMeshAndLODIndex(int32 MeshIndex, int32& OutLODIndex);
+	UE_API FMeshDescription* FindRawMeshAndLODIndex(int32 MeshIndex, int32& OutLODIndex);
 	/** Tries to retrieve a FMeshDescription for the given mesh and LOD index, if it can't it will try to find an entry for each LOD levle below InOutDesiredLODIndex */
-	FMeshDescription* TryFindRawMeshForLOD(int32 MeshIndex, int32& InOutDesiredLODIndex);
+	UE_API FMeshDescription* TryFindRawMeshForLOD(int32 MeshIndex, int32& InOutDesiredLODIndex);
 
 	/** Returns a const key/value iterator for the FMeshDescription entries */
-	TConstRawMeshIterator GetConstRawMeshIterator() const;
+	UE_API TConstRawMeshIterator GetConstRawMeshIterator() const;
 	/** Returns a non-const key/value iterator for the FMeshDescription entries */ 
-	TRawMeshIterator GetRawMeshIterator();
+	UE_API TRawMeshIterator GetRawMeshIterator();
 
 	/** Adds a record of what channel lightmap data is stored at */
-	void AddLightmapChannelRecord(int32 MeshIndex, int32 LODIndex, int32 LightmapChannelIndex);
+	UE_API void AddLightmapChannelRecord(int32 MeshIndex, int32 LODIndex, int32 LightmapChannelIndex);
 
 	/** Adds (unique) section to stored data */
-	int32 AddSection(const FSectionInfo& SectionInfo);
+	UE_API int32 AddSection(const FSectionInfo& SectionInfo);
 	/** Returns the number of unique sections */
-	int32 NumberOfUniqueSections() const;
+	UE_API int32 NumberOfUniqueSections() const;
 	/** Returns the material used by the unique section */
-	UMaterialInterface* GetMaterialForSectionIndex(int32 SectionIndex);
+	UE_API UMaterialInterface* GetMaterialForSectionIndex(int32 SectionIndex);
 	/** Returns the unique section instance */
-	const FSectionInfo& GetSection(int32 SectionIndex) const;
+	UE_API const FSectionInfo& GetSection(int32 SectionIndex) const;
 	/** Clears out unique section to be replaced with the baked material one */
-	void AddBakedMaterialSection(const FSectionInfo& SectionInfo);
+	UE_API void AddBakedMaterialSection(const FSectionInfo& SectionInfo);
 	
 	/** Add a material slot name for a unique material instance. */
-	void AddMaterialSlotName(UMaterialInterface *MaterialInterface, FName MaterialSlotName);
+	UE_API void AddMaterialSlotName(UMaterialInterface *MaterialInterface, FName MaterialSlotName);
 	/** Get the material slot name from a unique material instance. */
-	FName GetMaterialSlotName(UMaterialInterface *MaterialInterface) const;
+	UE_API FName GetMaterialSlotName(UMaterialInterface *MaterialInterface) const;
 
 	/** Adds a LOD index which will be part of the final merged mesh */
-	void AddLODIndex(int32 LODIndex);
+	UE_API void AddLODIndex(int32 LODIndex);
 	/** Retrieves number of LODs part of the final merged mesh */
-	int32 GetNumLODsForMergedMesh() const;
+	UE_API int32 GetNumLODsForMergedMesh() const;
 	/** Iterates over LOD indices for mesh */
-	TConstLODIndexIterator GetLODIndexIterator() const;
+	UE_API TConstLODIndexIterator GetLODIndexIterator() const;
 
 	/** Add number of lightmap pixels used for one of the Meshes */
-	void AddLightMapPixels(int32 Pixels);
+	UE_API void AddLightMapPixels(int32 Pixels);
 	/** Returns the texture dimension required to distribute all of the lightmap pixels */
-	int32 GetLightMapDimension() const;
+	UE_API int32 GetLightMapDimension() const;
 
 	/** Returns whether or not any raw mesh entry contains vertex colors for the specified LOD index */
-	bool DoesLODContainVertexColors(int32 LODIndex) const;
+	UE_API bool DoesLODContainVertexColors(int32 LODIndex) const;
 	/** Returns whether or not any raw mesh entry contains vertex colors for any LOD index */
-	bool DoesAnyLODContainVertexColors() const;
+	UE_API bool DoesAnyLODContainVertexColors() const;
 	/** Returns whether or not any raw mesh entry contains texture coordinates for the specified UV channel and LOD index */
-	bool DoesUVChannelContainData(int32 UVChannel, int32 LODIndex) const;
+	UE_API bool DoesUVChannelContainData(int32 UVChannel, int32 LODIndex) const;
 	/** Returns whether or not any raw mesh entry contains texture coordinates for the specified UV channel and any LOD index */
-	bool DoesUVChannelContainData(int32 UVChannel) const;
+	UE_API bool DoesUVChannelContainData(int32 UVChannel) const;
 	/** Returns whether or not the raw mesh entry for the given MeshLOD key requires unique UVs for baking out its material(s) */
-	bool DoesMeshLODRequireUniqueUVs(FMeshLODKey Key);
+	UE_API bool DoesMeshLODRequireUniqueUVs(FMeshLODKey Key);
 	/** Returns the first available UV channel across all raw mesh entries, which will be a good fit for the lightmap UV index in the final mesh */
-	int32 GetAvailableLightMapUVChannel() const;
+	UE_API int32 GetAvailableLightMapUVChannel() const;
 
-	void AddComponentToWedgeMapping(int32 MeshIndex, int32 LODIndex, uint32 WedgeIndex);
-	uint32 GetComponentToWedgeMappng(int32 MeshIndex, int32 LODIndex) const;
+	UE_API void AddComponentToWedgeMapping(int32 MeshIndex, int32 LODIndex, uint32 WedgeIndex);
+	UE_API uint32 GetComponentToWedgeMappng(int32 MeshIndex, int32 LODIndex) const;
 
-	double GetTextureSizeFromTargetTexelDensity(float InTargetTexelDensity) const;
+	UE_API double GetTextureSizeFromTargetTexelDensity(float InTargetTexelDensity) const;
 
 protected:
 	// Mesh / LOD index, RawMesh
@@ -221,3 +223,4 @@ protected:
 	TMap<FMeshLODKey, uint32> ComponentToWedgeOffsets;
 };
 
+#undef UE_API
diff --git a/Engine/Source/Developer/MeshMergeUtilities/Public/MeshMergeModule.h b/Engine/Source/Developer/MeshMergeUtilities/Public/MeshMergeModule.h
--- a/Engine/Source/Developer/MeshMergeUtilities/Public/MeshMergeModule.h
+++ b/Engine/Source/Developer/MeshMergeUtilities/Public/MeshMergeModule.h
@@ -6,7 +6,7 @@
 #include "Modules/ModuleInterface.h"
 #include "IMeshMergeUtilities.h"
 
-class MESHMERGEUTILITIES_API IMeshMergeModule : public IModuleInterface
+class IMeshMergeModule : public IModuleInterface
 {
 public:
 	virtual const IMeshMergeUtilities& GetUtilities() const = 0;
diff --git a/Engine/Source/Developer/MeshMergeUtilities/Public/StaticMeshComponentAdapter.h b/Engine/Source/Developer/MeshMergeUtilities/Public/StaticMeshComponentAdapter.h
--- a/Engine/Source/Developer/MeshMergeUtilities/Public/StaticMeshComponentAdapter.h
+++ b/Engine/Source/Developer/MeshMergeUtilities/Public/StaticMeshComponentAdapter.h
@@ -4,33 +4,35 @@
 
 #include "IMaterialBakingAdapter.h"
 
+#define UE_API MESHMERGEUTILITIES_API
+
 class UStaticMeshComponent;
 class UStaticMesh;
 
 /** Adapter which takes a Static Mesh Component instance to use for material baking (does not allow for changes to the underlying asset itself) */
-class MESHMERGEUTILITIES_API FStaticMeshComponentAdapter : public IMaterialBakingAdapter
+class FStaticMeshComponentAdapter : public IMaterialBakingAdapter
 {
 public:
-	FStaticMeshComponentAdapter(UStaticMeshComponent* InStaticMeshComponent);
+	UE_API FStaticMeshComponentAdapter(UStaticMeshComponent* InStaticMeshComponent);
 
 	/** Begin IMaterialBakingAdapter overrides */
-	virtual int32 GetNumberOfLODs() const override;
-	virtual void RetrieveRawMeshData(int32 LODIndex, FMeshDescription& InOutRawMesh, bool bPropogateMeshData) const override;
-	virtual void RetrieveMeshSections(int32 LODIndex, TArray<FSectionInfo>& InOutSectionInfo) const override;
-	virtual int32 GetMaterialIndex(int32 LODIndex, int32 SectionIndex) const override;	
-	virtual void ApplySettings(int32 LODIndex, FMeshData& InOutMeshData) const override;
-	virtual UPackage* GetOuter() const override;
-	virtual FString GetBaseName() const override;
-	virtual FName GetMaterialSlotName(int32 MaterialIndex) const override;
-	virtual FName GetImportedMaterialSlotName(int32 MaterialIndex) const override;
-	virtual void SetMaterial(int32 MaterialIndex, UMaterialInterface* Material) override;
-	virtual void RemapMaterialIndex(int32 LODIndex, int32 SectionIndex, int32 NewMaterialIndex) override;
-	virtual int32 AddMaterial(UMaterialInterface* Material) override;
-	virtual int32 AddMaterial(UMaterialInterface* Material, const FName& SlotName, const FName& ImportedSlotName) override;
-	virtual void UpdateUVChannelData() override;
-	virtual bool IsAsset() const override;
-	virtual int32 LightmapUVIndex() const override;
-	virtual FBoxSphereBounds GetBounds() const override;
+	UE_API virtual int32 GetNumberOfLODs() const override;
+	UE_API virtual void RetrieveRawMeshData(int32 LODIndex, FMeshDescription& InOutRawMesh, bool bPropogateMeshData) const override;
+	UE_API virtual void RetrieveMeshSections(int32 LODIndex, TArray<FSectionInfo>& InOutSectionInfo) const override;
+	UE_API virtual int32 GetMaterialIndex(int32 LODIndex, int32 SectionIndex) const override;	
+	UE_API virtual void ApplySettings(int32 LODIndex, FMeshData& InOutMeshData) const override;
+	UE_API virtual UPackage* GetOuter() const override;
+	UE_API virtual FString GetBaseName() const override;
+	UE_API virtual FName GetMaterialSlotName(int32 MaterialIndex) const override;
+	UE_API virtual FName GetImportedMaterialSlotName(int32 MaterialIndex) const override;
+	UE_API virtual void SetMaterial(int32 MaterialIndex, UMaterialInterface* Material) override;
+	UE_API virtual void RemapMaterialIndex(int32 LODIndex, int32 SectionIndex, int32 NewMaterialIndex) override;
+	UE_API virtual int32 AddMaterial(UMaterialInterface* Material) override;
+	UE_API virtual int32 AddMaterial(UMaterialInterface* Material, const FName& SlotName, const FName& ImportedSlotName) override;
+	UE_API virtual void UpdateUVChannelData() override;
+	UE_API virtual bool IsAsset() const override;
+	UE_API virtual int32 LightmapUVIndex() const override;
+	UE_API virtual FBoxSphereBounds GetBounds() const override;
 	/** End IMaterialBakingAdapter overrides */
 
 protected:
@@ -38,3 +40,5 @@ protected:
 	UStaticMesh* StaticMesh;
 	int32 NumLODs;
 };
+
+#undef UE_API
diff --git a/Engine/Source/Developer/OutputLog/Public/OutputLogSettings.h b/Engine/Source/Developer/OutputLog/Public/OutputLogSettings.h
--- a/Engine/Source/Developer/OutputLog/Public/OutputLogSettings.h
+++ b/Engine/Source/Developer/OutputLog/Public/OutputLogSettings.h
@@ -81,8 +81,8 @@ struct FOutputLogFilterSettings
 /**
  * Implements the Editor style settings.
  */
-UCLASS(config=EditorPerProjectUserSettings)
-class OUTPUTLOG_API UOutputLogSettings : public UObject
+UCLASS(MinimalAPI, config=EditorPerProjectUserSettings)
+class UOutputLogSettings : public UObject
 {
 	GENERATED_BODY()
 
diff --git a/Engine/Source/Developer/S3Client/Internal/S3/S3Client.h b/Engine/Source/Developer/S3Client/Internal/S3/S3Client.h
--- a/Engine/Source/Developer/S3Client/Internal/S3/S3Client.h
+++ b/Engine/Source/Developer/S3Client/Internal/S3/S3Client.h
@@ -73,15 +73,15 @@ private:
 /**
  * Container for named credentials.
  */
-class UE_API FS3CredentialsProfileStore
+class FS3CredentialsProfileStore
 {
 public:
 	/** Returns the first credentials loaded into the profile store. */
-	FS3ClientCredentials GetDefault() const;
+	UE_API FS3ClientCredentials GetDefault() const;
 	/** Get credentials for the specified profile name. */
-	bool TryGetCredentials(const FString& ProfileName, FS3ClientCredentials& OutCredentials) const;
+	UE_API bool TryGetCredentials(const FString& ProfileName, FS3ClientCredentials& OutCredentials) const;
 	/** Reads named credentials from an .ini file. */
-	static FS3CredentialsProfileStore FromFile(const FString& FileName, FString* OutError = nullptr);
+	static UE_API FS3CredentialsProfileStore FromFile(const FString& FileName, FString* OutError = nullptr);
 
 private:
 	TMap<FString, FS3ClientCredentials> Credentials;
@@ -124,7 +124,7 @@ public:
 
 public:
 
-	UE_API ~FS3Response() = default;
+	~FS3Response() = default;
 
 	UE_API FS3Response& operator=(const FS3Response& Other);
 	UE_API FS3Response& operator=(FS3Response&& Other);
@@ -260,17 +260,17 @@ struct FS3ClientConfig
 /**
  * A simple HTTP(s) client for down, uploading and listing data objects from Amazon S3.
  */
-class UE_API FS3Client
+class FS3Client
 {
 	class FConnectionPool;
 	class FS3Request;
 
 public:
 	/** Creates a new instance from the specified configuration and credentials. */
-	explicit FS3Client(const FS3ClientConfig& ClientConfig, const FS3ClientCredentials& BasicCredentials);
+	UE_API explicit FS3Client(const FS3ClientConfig& ClientConfig, const FS3ClientCredentials& BasicCredentials);
 	FS3Client(const FS3Client&) = delete;
 	FS3Client(FS3Client&&) = delete;
-	~FS3Client();
+	UE_API ~FS3Client();
 
 	FS3Client& operator=(const FS3Client&) = delete;
 	FS3Client& operator=(FS3Client&&) = delete;
@@ -281,21 +281,21 @@ public:
 	const FS3ClientConfig& GetConfig() const { return Config; }
 
 	/** Download an object associated with the specified request parameters. */
-	FS3GetObjectResponse GetObject(const FS3GetObjectRequest& Request);
+	UE_API FS3GetObjectResponse GetObject(const FS3GetObjectRequest& Request);
 	/** Download object meta data. */
-	FS3HeadObjectResponse HeadObject(const FS3HeadObjectRequest& Request);
+	UE_API FS3HeadObjectResponse HeadObject(const FS3HeadObjectRequest& Request);
 	/** List all objects described by the specified request parameters. */
-	FS3ListObjectResponse ListObjects(const FS3ListObjectsRequest& Request);
+	UE_API FS3ListObjectResponse ListObjects(const FS3ListObjectsRequest& Request);
 	/** Upload an object described by the specified request parameters. */
-	FS3PutObjectResponse PutObject(const FS3PutObjectRequest& Request);
+	UE_API FS3PutObjectResponse PutObject(const FS3PutObjectRequest& Request);
 	/** Retries uploading an object until succeeded or max attempts has been reached. */
-	FS3PutObjectResponse TryPutObject(const FS3PutObjectRequest& Request, int32 MaxAttempts = 3, float Delay = 0.5f);
+	UE_API FS3PutObjectResponse TryPutObject(const FS3PutObjectRequest& Request, int32 MaxAttempts = 3, float Delay = 0.5f);
 	/** Delete an object associated with the specified request parameters. */
-	FS3DeleteObjectResponse DeleteObject(const FS3DeleteObjectRequest& Request);
+	UE_API FS3DeleteObjectResponse DeleteObject(const FS3DeleteObjectRequest& Request);
 
 private:
-	void Setup(FS3Request& Request);
-	void Teardown(FS3Request& Request);
+	UE_API void Setup(FS3Request& Request);
+	UE_API void Teardown(FS3Request& Request);
 
 	FS3ClientConfig Config;
 	FS3ClientCredentials Credentials;
diff --git a/Engine/Source/Developer/SessionFrontend/Public/Widgets/Browser/SSessionBrowser.h b/Engine/Source/Developer/SessionFrontend/Public/Widgets/Browser/SSessionBrowser.h
--- a/Engine/Source/Developer/SessionFrontend/Public/Widgets/Browser/SSessionBrowser.h
+++ b/Engine/Source/Developer/SessionFrontend/Public/Widgets/Browser/SSessionBrowser.h
@@ -14,6 +14,8 @@
 #include "Widgets/Views/STableRow.h"
 #include "Widgets/Views/STreeView.h"
 
+#define UE_API SESSIONFRONTEND_API
+
 class FSessionBrowserGroupTreeItem;
 class FSessionBrowserTreeItem;
 class SSessionBrowserCommandBar;
@@ -21,7 +23,7 @@ class SSessionBrowserCommandBar;
 /**
  * Implements a Slate widget for browsing active game sessions.
  */
-class SESSIONFRONTEND_API SSessionBrowser
+class SSessionBrowser
 	: public SCompoundWidget
 {
 public:
@@ -32,7 +34,7 @@ public:
 public:
 
 	/** Destructor. */
-	~SSessionBrowser();
+	UE_API ~SSessionBrowser();
 
 public:
 
@@ -42,7 +44,7 @@ public:
 	 * @param InArgs The declaration data for this widget.
 	 * @param InSessionManager The session manager to use.
 	 */
-	void Construct( const FArguments& InArgs, TSharedRef<ISessionManager> InSessionManager );
+	UE_API void Construct( const FArguments& InArgs, TSharedRef<ISessionManager> InSessionManager );
 
 protected:
 
@@ -51,51 +53,51 @@ protected:
 	 *
 	 * @param Item The tree view item to expand.
 	 */
-	void ExpandItem(const TSharedPtr<FSessionBrowserTreeItem>& Item);
+	UE_API void ExpandItem(const TSharedPtr<FSessionBrowserTreeItem>& Item);
 
 	/** Filters the session tree. */
-	void FilterSessions();
+	UE_API void FilterSessions();
 
 	/** Adds items for this session in the tree. */
-	void AddInstanceItemToTree(TSharedPtr<FSessionBrowserTreeItem>& SessionItem, const TSharedPtr<FSessionBrowserTreeItem>& InstanceItem, const TSharedPtr<ISessionInstanceInfo>& InstanceInfo);
+	UE_API void AddInstanceItemToTree(TSharedPtr<FSessionBrowserTreeItem>& SessionItem, const TSharedPtr<FSessionBrowserTreeItem>& InstanceItem, const TSharedPtr<ISessionInstanceInfo>& InstanceInfo);
 
 	/** Reloads the sessions list. */
-	void ReloadSessions();
+	UE_API void ReloadSessions();
 
 private:
 
 	/** Callback for changing the selection state of an instance. */
-	void HandleSessionManagerInstanceSelectionChanged(const TSharedPtr<ISessionInstanceInfo>& Instance, bool Selected);
+	UE_API void HandleSessionManagerInstanceSelectionChanged(const TSharedPtr<ISessionInstanceInfo>& Instance, bool Selected);
 
 	/** Callback for changing the selected session in the session manager. */
-	void HandleSessionManagerSelectedSessionChanged(const TSharedPtr<ISessionInfo>& SelectedSession);
+	UE_API void HandleSessionManagerSelectedSessionChanged(const TSharedPtr<ISessionInfo>& SelectedSession);
 
 	/** Callback for updating the session list in the session manager. */
-	void HandleSessionManagerSessionsUpdated();
+	UE_API void HandleSessionManagerSessionsUpdated();
 
 	/** Callback from the session manager to notify there's a new session instance. */
-	void HandleSessionManagerInstanceDiscovered(const TSharedRef<ISessionInfo>& OwnerSession, const TSharedRef<ISessionInstanceInfo>& DiscoveredInstance);
+	UE_API void HandleSessionManagerInstanceDiscovered(const TSharedRef<ISessionInfo>& OwnerSession, const TSharedRef<ISessionInstanceInfo>& DiscoveredInstance);
 
 	/** Callback for getting the tool tip text of a session tree row. */
-	FText HandleSessionTreeRowGetToolTipText(TSharedPtr<FSessionBrowserTreeItem> Item) const;
+	UE_API FText HandleSessionTreeRowGetToolTipText(TSharedPtr<FSessionBrowserTreeItem> Item) const;
 
 	/** Callback for session tree view selection changes. */
-	void HandleSessionTreeViewExpansionChanged(TSharedPtr<FSessionBrowserTreeItem> TreeItem, bool bIsExpanded);
+	UE_API void HandleSessionTreeViewExpansionChanged(TSharedPtr<FSessionBrowserTreeItem> TreeItem, bool bIsExpanded);
 
 	/** Callback for generating a row widget in the session tree view. */
-	TSharedRef<ITableRow> HandleSessionTreeViewGenerateRow(TSharedPtr<FSessionBrowserTreeItem> Item, const TSharedRef<STableViewBase>& OwnerTable);
+	UE_API TSharedRef<ITableRow> HandleSessionTreeViewGenerateRow(TSharedPtr<FSessionBrowserTreeItem> Item, const TSharedRef<STableViewBase>& OwnerTable);
 
 	/** Callback for getting the children of a node in the session tree view. */
-	void HandleSessionTreeViewGetChildren(TSharedPtr<FSessionBrowserTreeItem> Item, TArray<TSharedPtr<FSessionBrowserTreeItem>>& OutChildren);
+	UE_API void HandleSessionTreeViewGetChildren(TSharedPtr<FSessionBrowserTreeItem> Item, TArray<TSharedPtr<FSessionBrowserTreeItem>>& OutChildren);
 
 	/** Callback for session tree view selection changes. */
-	void HandleSessionTreeViewSelectionChanged(const TSharedPtr<FSessionBrowserTreeItem> Item, ESelectInfo::Type SelectInfo);
+	UE_API void HandleSessionTreeViewSelectionChanged(const TSharedPtr<FSessionBrowserTreeItem> Item, ESelectInfo::Type SelectInfo);
 
 	/** Callback for clicking the 'Terminate Session' button. */
-	FReply HandleTerminateSessionButtonClicked();
+	UE_API FReply HandleTerminateSessionButtonClicked();
 
 	/** Callback for getting the enabled state of the 'Terminate Session' button. */
-	bool HandleTerminateSessionButtonIsEnabled() const;
+	UE_API bool HandleTerminateSessionButtonIsEnabled() const;
 
 private:
 
@@ -143,3 +145,5 @@ private:
 	/** True if we should set the default selection the next time the tree view if refreshed. */
 	bool bCanSetDefaultSelection;
 };
+
+#undef UE_API
diff --git a/Engine/Source/Developer/SharedSettingsWidgets/Public/ManifestUpdateHelper.h b/Engine/Source/Developer/SharedSettingsWidgets/Public/ManifestUpdateHelper.h
--- a/Engine/Source/Developer/SharedSettingsWidgets/Public/ManifestUpdateHelper.h
+++ b/Engine/Source/Developer/SharedSettingsWidgets/Public/ManifestUpdateHelper.h
@@ -5,36 +5,40 @@
 #include "CoreMinimal.h"
 #include "Misc/FileHelper.h"
 
+#define UE_API SHAREDSETTINGSWIDGETS_API
+
 /////////////////////////////////////////////////////
 // FManifestUpdateHelper
 
 // This is a utility class used to update individual sections of an XML manifest or .plist in raw string form.
 // It is formatting sensitive and will fail if things are formatted in an unexpected manner.
 // This is a stopgap measure, and will be replaced using a proper XML parser when it is ready for use.
-class SHAREDSETTINGSWIDGETS_API FManifestUpdateHelper
+class FManifestUpdateHelper
 {
 public:
-	FManifestUpdateHelper(const FString& InFilename);
+	UE_API FManifestUpdateHelper(const FString& InFilename);
 
 	// Replaces the text in InOutString between MatchPrefix and MatchSuffix with NewInfix, returning true if it was found
-	static bool ReplaceStringPortion(FString& InOutString, const FString& MatchPrefix, const FString& MatchSuffix, const FString& NewInfix);
+	static UE_API bool ReplaceStringPortion(FString& InOutString, const FString& MatchPrefix, const FString& MatchSuffix, const FString& NewInfix);
 
 	// Checks if a key exists matching MatchPrefix
-	bool HasKey(const FString& MatchPrefix);
+	UE_API bool HasKey(const FString& MatchPrefix);
 
 	// Replace a key in the manifest (between MatchPrefix and MatchSuffix) with NewInfix
-	void ReplaceKey(const FString& MatchPrefix, const FString& MatchSuffix, const FString& NewInfix);
+	UE_API void ReplaceKey(const FString& MatchPrefix, const FString& MatchSuffix, const FString& NewInfix);
 
 	// Finalizes the updater, and writes it back to the file, returning true if successful, and false if there were any errors
-	bool Finalize(const FString& TargetFilename, bool bShowNotifyOnFailure = true, FFileHelper::EEncodingOptions EncodingOption = FFileHelper::EEncodingOptions::AutoDetect);
+	UE_API bool Finalize(const FString& TargetFilename, bool bShowNotifyOnFailure = true, FFileHelper::EEncodingOptions EncodingOption = FFileHelper::EEncodingOptions::AutoDetect);
 
 	// Returns the first error message
 	FText GetFirstErrorMessage() const { return FirstErrorMessage; }
 
 private:
-	void WriteError(FText NewError);
+	UE_API void WriteError(FText NewError);
 private:
 	FText FirstErrorMessage;
 	FString ManifestString;
 	bool bManifestDirty;
 };
+
+#undef UE_API
diff --git a/Engine/Source/Developer/SharedSettingsWidgets/Public/SExternalImageReference.h b/Engine/Source/Developer/SharedSettingsWidgets/Public/SExternalImageReference.h
--- a/Engine/Source/Developer/SharedSettingsWidgets/Public/SExternalImageReference.h
+++ b/Engine/Source/Developer/SharedSettingsWidgets/Public/SExternalImageReference.h
@@ -8,6 +8,8 @@
 #include "PropertyHandle.h"
 #include "IExternalImagePickerModule.h"
 
+#define UE_API SHAREDSETTINGSWIDGETS_API
+
 /**
  * Delegate fired before an image has been copied
  * @param	InChosenImage	The path to the image that has been picked by the user.
@@ -28,7 +30,7 @@ DECLARE_DELEGATE_RetVal_OneParam(bool, FOnPostExternalImageCopy, const FString&
 // This widget shows an external image preview of a per-project configurable image
 // (one where the engine provides a default, but each project may have its own override)
 
-class SHAREDSETTINGSWIDGETS_API SExternalImageReference : public SCompoundWidget
+class SExternalImageReference : public SCompoundWidget
 {
 	SLATE_BEGIN_ARGS(SExternalImageReference)
 		: _FileDescription(NSLOCTEXT("SExternalImageReference", "FileDescription", "External Image"))
@@ -82,11 +84,11 @@ class SHAREDSETTINGSWIDGETS_API SExternalImageReference : public SCompoundWidget
 
 public:
 
-	void Construct(const FArguments& InArgs, const FString& InBaseFilename, const FString& InOverrideFilename);
+	UE_API void Construct(const FArguments& InArgs, const FString& InBaseFilename, const FString& InOverrideFilename);
 
 protected:
 	/** Delegate handler for when an image is picked */
-	bool HandleExternalImagePicked(const FString& InChosenImage, const FString& InTargetImage);
+	UE_API bool HandleExternalImagePicked(const FString& InChosenImage, const FString& InTargetImage);
 
 protected:
 
@@ -117,3 +119,5 @@ protected:
 	/** If true, the previous target image will be deleted if the file extension changes */
 	bool bDeletePreviousTargetWhenExtensionChanges;
 };
+
+#undef UE_API
diff --git a/Engine/Source/Developer/SharedSettingsWidgets/Public/SHyperlinkLaunchURL.h b/Engine/Source/Developer/SharedSettingsWidgets/Public/SHyperlinkLaunchURL.h
--- a/Engine/Source/Developer/SharedSettingsWidgets/Public/SHyperlinkLaunchURL.h
+++ b/Engine/Source/Developer/SharedSettingsWidgets/Public/SHyperlinkLaunchURL.h
@@ -6,12 +6,14 @@
 #include "Widgets/DeclarativeSyntaxSupport.h"
 #include "Widgets/SCompoundWidget.h"
 
+#define UE_API SHAREDSETTINGSWIDGETS_API
+
 /////////////////////////////////////////////////////
 // SHyperlinkLaunchURL
 
 // This widget is a hyperlink that launches an external URL when clicked
 
-class SHAREDSETTINGSWIDGETS_API SHyperlinkLaunchURL : public SCompoundWidget
+class SHyperlinkLaunchURL : public SCompoundWidget
 {
 	SLATE_BEGIN_ARGS(SHyperlinkLaunchURL)
 		{}
@@ -22,11 +24,13 @@ class SHAREDSETTINGSWIDGETS_API SHyperlinkLaunchURL : public SCompoundWidget
 	SLATE_END_ARGS()
 
 public:
-	void Construct(const FArguments& InArgs, const FString& InDestinationURL);
+	UE_API void Construct(const FArguments& InArgs, const FString& InDestinationURL);
 
 protected:
-	void OnNavigate();
+	UE_API void OnNavigate();
 
 protected:
 	FString DestinationURL;
 };
+
+#undef UE_API
diff --git a/Engine/Source/Developer/SharedSettingsWidgets/Public/SPlatformSetupMessage.h b/Engine/Source/Developer/SharedSettingsWidgets/Public/SPlatformSetupMessage.h
--- a/Engine/Source/Developer/SharedSettingsWidgets/Public/SPlatformSetupMessage.h
+++ b/Engine/Source/Developer/SharedSettingsWidgets/Public/SPlatformSetupMessage.h
@@ -12,12 +12,14 @@
 #include "Widgets/SWidget.h"
 #include "Widgets/SCompoundWidget.h"
 
+#define UE_API SHAREDSETTINGSWIDGETS_API
+
 /////////////////////////////////////////////////////
 // SPlatformSetupMessage
 
 // This widget displays a setup message indicating if the game project is configured for a platform or not
 
-class SHAREDSETTINGSWIDGETS_API SPlatformSetupMessage : public SCompoundWidget
+class SPlatformSetupMessage : public SCompoundWidget
 {
 	enum ESetupState
 	{
@@ -40,31 +42,33 @@ class SHAREDSETTINGSWIDGETS_API SPlatformSetupMessage : public SCompoundWidget
 	SLATE_END_ARGS()
 
 public:
-	void Construct(const FArguments& InArgs, const FString& InTargetFilename);
+	UE_API void Construct(const FArguments& InArgs, const FString& InTargetFilename);
 
-	TAttribute<bool> GetReadyToGoAttribute() const;
+	UE_API TAttribute<bool> GetReadyToGoAttribute() const;
 
 	// SWidget interface
-	virtual void Tick( const FGeometry& AllottedGeometry, const double InCurrentTime, const float InDeltaTime ) override;
+	UE_API virtual void Tick( const FGeometry& AllottedGeometry, const double InCurrentTime, const float InDeltaTime ) override;
 	// End of Swidget interface
 private:
-	int32 GetSetupStateAsInt() const;
-	bool IsReadyToGo() const;
-	FSlateColor GetBorderColor() const;
+	UE_API int32 GetSetupStateAsInt() const;
+	UE_API bool IsReadyToGo() const;
+	UE_API FSlateColor GetBorderColor() const;
 
-	TSharedRef<SWidget> MakeRow(FName IconName, FText Message, FText ButtonMessage);
+	UE_API TSharedRef<SWidget> MakeRow(FName IconName, FText Message, FText ButtonMessage);
 
-	FReply OnButtonPressed();
+	UE_API FReply OnButtonPressed();
 
-	void OnSourceControlOperationComplete(const FSourceControlOperationRef& InOperation, ECommandResult::Type InResult);
+	UE_API void OnSourceControlOperationComplete(const FSourceControlOperationRef& InOperation, ECommandResult::Type InResult);
 
 	// Returns the setup state of a specified file
-	ESetupState GetSetupStateBasedOnFile(bool bInitStatus);
+	UE_API ESetupState GetSetupStateBasedOnFile(bool bInitStatus);
 
 	// Updates the cache CachedSetupState 
-	void UpdateCache(bool bForceUpdate);
+	UE_API void UpdateCache(bool bForceUpdate);
 private:
 	FString TargetFilename;
 	ESetupState CachedSetupState;
 	FSimpleDelegate OnSetupClicked;
 };
+
+#undef UE_API
diff --git a/Engine/Source/Developer/SharedSettingsWidgets/Public/SSettingsEditorCheckoutNotice.h b/Engine/Source/Developer/SharedSettingsWidgets/Public/SSettingsEditorCheckoutNotice.h
--- a/Engine/Source/Developer/SharedSettingsWidgets/Public/SSettingsEditorCheckoutNotice.h
+++ b/Engine/Source/Developer/SharedSettingsWidgets/Public/SSettingsEditorCheckoutNotice.h
@@ -10,6 +10,8 @@
 #include "Input/Reply.h"
 #include "Widgets/SCompoundWidget.h"
 
+#define UE_API SHAREDSETTINGSWIDGETS_API
+
 namespace SettingsHelpers
 {
 	SHAREDSETTINGSWIDGETS_API bool IsSourceControlled(const FString& InFileToCheckOut, bool bForceSourceControlUpdate = false);
@@ -21,7 +23,7 @@ namespace SettingsHelpers
 /**
  * Implements a widget that provides a notice for files that need to be checked out.
  */
-class SHAREDSETTINGSWIDGETS_API SSettingsEditorCheckoutNotice : public SCompoundWidget
+class SSettingsEditorCheckoutNotice : public SCompoundWidget
 {
 public:
 
@@ -40,14 +42,14 @@ public:
 	 *
 	 * @param InArgs The construction arguments.
 	 */
-	void Construct( const FArguments& InArgs );
+	UE_API void Construct( const FArguments& InArgs );
 
 	// SWidget interface
-	virtual void Tick( const FGeometry& AllottedGeometry, const double InCurrentTime, const float InDeltaTime ) override;
+	UE_API virtual void Tick( const FGeometry& AllottedGeometry, const double InCurrentTime, const float InDeltaTime ) override;
 	// End of SWidget interface
 
 	// Return true if the watched file can be modified (as of the last update)
-	bool IsUnlocked() const;
+	UE_API bool IsUnlocked() const;
 
 	// Force an update on the next tick
 	void Invalidate()
@@ -58,16 +60,16 @@ public:
 private:
 
 	// Callback for clicking the 'Check Out' button.
-	FReply HandleCheckOutButtonClicked( );
+	UE_API FReply HandleCheckOutButtonClicked( );
 
 	// Callback for getting the text of the 'Check Out' button.
-	FText HandleCheckOutButtonText( ) const;
+	UE_API FText HandleCheckOutButtonText( ) const;
 
 	// Callback for getting the tool tip text of the 'Check Out' button.
-	FText HandleCheckOutButtonToolTip( ) const;
+	UE_API FText HandleCheckOutButtonToolTip( ) const;
 
 	// Callback for determining the visibility of the check-out button.
-	EVisibility HandleCheckOutButtonVisibility( ) const;
+	UE_API EVisibility HandleCheckOutButtonVisibility( ) const;
 
 	// Callback for getting the widget index for the notice switcher.
 	int32 HandleNoticeSwitcherWidgetIndex( ) const
@@ -76,16 +78,16 @@ private:
 	}
 
 	// Callback for getting the status text when the config is locked
-	FText HandleLockedStatusText() const;
+	UE_API FText HandleLockedStatusText() const;
 
 	// Callback for getting the status text when the config is unlocked
-	FText HandleUnlockedStatusText() const;
+	UE_API FText HandleUnlockedStatusText() const;
 
 	// Callback for getting the visibility of the source control throbber
-	EVisibility HandleThrobberVisibility() const;
+	UE_API EVisibility HandleThrobberVisibility() const;
 
 	// Callback to get the background color of the checkout wanring
-	FSlateColor GetLockedStatusTextColor() const;
+	UE_API FSlateColor GetLockedStatusTextColor() const;
 private:
 
 	// Called when the file was modified externally
@@ -103,3 +105,5 @@ private:
 	// Holds a timer for checking whether the section's configuration file needs to be checked out.
 	double LastDefaultConfigCheckOutTime;
 };
+
+#undef UE_API
diff --git a/Engine/Source/Developer/SharedSettingsWidgets/Public/ShaderFormatsPropertyDetails.h b/Engine/Source/Developer/SharedSettingsWidgets/Public/ShaderFormatsPropertyDetails.h
--- a/Engine/Source/Developer/SharedSettingsWidgets/Public/ShaderFormatsPropertyDetails.h
+++ b/Engine/Source/Developer/SharedSettingsWidgets/Public/ShaderFormatsPropertyDetails.h
@@ -5,6 +5,8 @@
 #include "CoreMinimal.h"
 #include "DetailLayoutBuilder.h"
 
+#define UE_API SHAREDSETTINGSWIDGETS_API
+
 class IDetailLayoutBuilder;
 class IPropertyHandle;
 class ITargetPlatform;
@@ -15,7 +17,7 @@ enum class ECheckBoxState : uint8;
 /**
  * Helper which implements details panel customizations for a device profiles parent property
  */
-class SHAREDSETTINGSWIDGETS_API FShaderFormatsPropertyDetails
+class FShaderFormatsPropertyDetails
 : public TSharedFromThis<FShaderFormatsPropertyDetails>
 {
 public:
@@ -31,16 +33,16 @@ public:
 	 * @param InProperty - The category name to override
 	 * @param InTitle - Title for display
 	 */
-	FShaderFormatsPropertyDetails(IDetailLayoutBuilder* InDetailBuilder, const TCHAR* InProperty = TEXT("TargetedRHIs"), const TCHAR* InTitle = TEXT("Targeted RHIs"));
+	UE_API FShaderFormatsPropertyDetails(IDetailLayoutBuilder* InDetailBuilder, const TCHAR* InProperty = TEXT("TargetedRHIs"), const TCHAR* InTitle = TEXT("Targeted RHIs"));
 
 	/** Simple delegate for updating shader version warning */
-	void SetOnUpdateShaderWarning(const FSimpleDelegate& Delegate);
+	UE_API void SetOnUpdateShaderWarning(const FSimpleDelegate& Delegate);
 	
 	/** Create the UI to select which windows shader formats we are targeting */
-	void CreateTargetShaderFormatsPropertyView(ITargetPlatformSettings* TargetPlatform, GetFriendlyNameFromRHINameFnc* FriendlyNameFnc, FilterShaderPlatformFnc* FilterShaderPlatformFunc = nullptr, ECategoryPriority::Type InPriority = ECategoryPriority::Default);
+	UE_API void CreateTargetShaderFormatsPropertyView(ITargetPlatformSettings* TargetPlatform, GetFriendlyNameFromRHINameFnc* FriendlyNameFnc, FilterShaderPlatformFnc* FilterShaderPlatformFunc = nullptr, ECategoryPriority::Type InPriority = ECategoryPriority::Default);
 	
 	UE_DEPRECATED(5.5, "CreateTargetShaderFormatsPropertyView now takes a ITargetPlatformSettings* as first argument instead of ITargetPlatform*. Please change your callback function as this is a breaking change.")
-	void CreateTargetShaderFormatsPropertyView(ITargetPlatform* TargetPlatform, GetFriendlyNameFromRHINameFnc* FriendlyNameFnc, FilterShaderPlatformFnc* FilterShaderPlatformFunc = nullptr, ECategoryPriority::Type InPriority = ECategoryPriority::Default);
+	UE_API void CreateTargetShaderFormatsPropertyView(ITargetPlatform* TargetPlatform, GetFriendlyNameFromRHINameFnc* FriendlyNameFnc, FilterShaderPlatformFnc* FilterShaderPlatformFunc = nullptr, ECategoryPriority::Type InPriority = ECategoryPriority::Default);
 
 	UE_DEPRECATED(5.1, "CreateTargetShaderFormatsPropertyView now gets RHI names via FName instead of FString. Please change your callback function as this is a breaking change.")
 	void CreateTargetShaderFormatsPropertyView(ITargetPlatform* TargetPlatform, Deprecated_GetFriendlyNameFromRHINameFnc* FriendlyNameFnc)
@@ -51,12 +53,12 @@ public:
 	 * @param InRHIName - The input RHI to check
 	 * @returns Whether this RHI is currently enabled
 	 */
-	ECheckBoxState IsTargetedRHIChecked(FName InRHIName) const;
+	UE_API ECheckBoxState IsTargetedRHIChecked(FName InRHIName) const;
 	
 private:
 	
 	// Supported/Targeted RHI check boxes
-	void OnTargetedRHIChanged(ECheckBoxState InNewValue, FName InRHIName);
+	UE_API void OnTargetedRHIChanged(ECheckBoxState InNewValue, FName InRHIName);
 	
 private:
 	
@@ -75,4 +77,6 @@ private:
 
 	/** Preserve shader format order when writing to property */
 	TMap<FName, int> ShaderFormatOrder;
-};
\ No newline at end of file
+};
+
+#undef UE_API
