World Partition - HLOD: Disable megalight in MaterialUtilities RenderSceneToTexture()
* Leaving megalight on might result in debug text being displayed in the captured texture ("MegaLights is enabled, but has no ray tracing data and won't operate correctly")

#Synced-CL 51374974

diff --git a/Engine/Source/Developer/MaterialUtilities/Private/MaterialUtilities.cpp b/Engine/Source/Developer/MaterialUtilities/Private/MaterialUtilities.cpp
--- a/Engine/Source/Developer/MaterialUtilities/Private/MaterialUtilities.cpp
+++ b/Engine/Source/Developer/MaterialUtilities/Private/MaterialUtilities.cpp
@@ -853,6 +853,7 @@ static void RenderSceneToTexture(
 	ViewFamily.EngineShowFlags.SetVisualizeBuffer(true);
 	ViewFamily.EngineShowFlags.SetTonemapper(false);
 	ViewFamily.EngineShowFlags.SetScreenPercentage(false);
+	ViewFamily.EngineShowFlags.SetMegaLights(false);
 
 	FSceneViewStateReference ViewState;
 	ViewState.Allocate(Scene->GetFeatureLevel());
